Fix link error unresolved external symbol RHIGetTrackedResourceStats in shipping_editor

diff --git a/Engine/Source/Editor/RenderResourceViewer/Private/SRenderResourceViewer.cpp b/Engine/Source/Editor/RenderResourceViewer/Private/SRenderResourceViewer.cpp
--- a/Engine/Source/Editor/RenderResourceViewer/Private/SRenderResourceViewer.cpp
+++ b/Engine/Source/Editor/RenderResourceViewer/Private/SRenderResourceViewer.cpp
@@ -245,7 +245,9 @@ void SRenderResourceViewerWidget::RefreshNodes()
 
 	// Retrieve the whole list of RHI tracked resources
 	TArray<TSharedPtr<FRHIResourceStats>> RHIResources;
+#if RHI_ENABLE_RESOURCE_INFO
 	RHIGetTrackedResourceStats(RHIResources);
+#endif
 
 	// Filter the resource list by search keyword and flags
 	ResourceInfos.Reset();
