Add owner column to Render Resource Viewer, add owner path name tracking to skeletal & static mesh buffers

diff --git a/Engine/Source/Editor/RenderResourceViewer/Private/SRenderResourceViewer.cpp b/Engine/Source/Editor/RenderResourceViewer/Private/SRenderResourceViewer.cpp
--- a/Engine/Source/Editor/RenderResourceViewer/Private/SRenderResourceViewer.cpp
+++ b/Engine/Source/Editor/RenderResourceViewer/Private/SRenderResourceViewer.cpp
@@ -15,6 +15,7 @@ namespace RenderResourceViewerInternal
 	const FName ColumnType("Type");
 	const FName ColumnSize("Size");
 	const FName ColumnFlags("Flags");
+	const FName ColumnOwner("Owner");
 	
 	// Format the size to nearest size unit of Byte/KB/MB
 	FString GetFormatedSize(uint64 SizeInBytes)
@@ -55,7 +56,7 @@ namespace RenderResourceViewerInternal
 			if (Column == RenderResourceViewerInternal::ColumnName)
 			{
 				return SNew(STextBlock)
-					.Text(FText::FromString(Info->Name));
+					.Text(FText::FromName(Info->Name));
 			}
 			else if (Column == RenderResourceViewerInternal::ColumnType)
 			{
@@ -72,6 +73,11 @@ namespace RenderResourceViewerInternal
 				return SNew(STextBlock)
 					.Text(FText::FromString(Info->Flags));
 			}
+			else if (Column == RenderResourceViewerInternal::ColumnOwner)
+			{
+				return SNew(STextBlock)
+					.Text(FText::FromName(Info->OwnerName));
+			}
 			return SNullWidget::NullWidget;
 		}
 	};
@@ -95,6 +101,13 @@ namespace RenderResourceViewerInternal
 			.OnCheckStateChanged(this, &OnCheckStateChangedFunc) \
 		]
 
+#define RENDER_RESOURCE_VIEWER_ADD_COLUMN(Name, Width, Label) \
+	+ SHeaderRow::Column(Name) \
+		.FillWidth(Width) \
+		.DefaultLabel(Label) \
+		.SortMode(this, &SRenderResourceViewerWidget::GetColumnSortMode, Name) \
+		.OnSort(this, &SRenderResourceViewerWidget::OnColumnSortModeChanged)
+
 void SRenderResourceViewerWidget::Construct(const FArguments& InArgs, const TSharedRef<SDockTab>& ConstructUnderMajorTab, const TSharedPtr<SWindow>& ConstructUnderWindow)
 {
 	SortByColumn = RenderResourceViewerInternal::ColumnSize;
@@ -150,8 +163,8 @@ void SRenderResourceViewerWidget::Construct(const FArguments& InArgs, const TSha
 				.AutoHeight()
 				[
 					SAssignNew(FilterTextBox, SSearchBox)
-					.HintText(LOCTEXT("FilterTextBoxHint", "Search resources by name."))
-					.ToolTipText(LOCTEXT("FilterTextBoxToolTip", "Type here to filter the list of render resources by name."))
+					.HintText(LOCTEXT("FilterTextBoxHint", "Search resources by name or owner."))
+					.ToolTipText(LOCTEXT("FilterTextBoxToolTip", "Type here to filter the list of render resources by name or owner."))
 					.OnTextChanged(this, &SRenderResourceViewerWidget::FilterTextBox_OnTextChanged)
 				]
 
@@ -207,29 +220,11 @@ void SRenderResourceViewerWidget::Construct(const FArguments& InArgs, const TSha
 					(
 						SNew(SHeaderRow)
 
-						+ SHeaderRow::Column(RenderResourceViewerInternal::ColumnName)
-						.FillWidth(0.6f)
-						.DefaultLabel(LOCTEXT("NameColumn", "Resource Name"))
-						.SortMode(this, &SRenderResourceViewerWidget::GetColumnSortMode, RenderResourceViewerInternal::ColumnName)
-						.OnSort(this, &SRenderResourceViewerWidget::OnColumnSortModeChanged)
-
-						+ SHeaderRow::Column(RenderResourceViewerInternal::ColumnType)
-						.FillWidth(0.1f)
-						.DefaultLabel(LOCTEXT("TypeColumn", "Type"))
-						.SortMode(this, &SRenderResourceViewerWidget::GetColumnSortMode, RenderResourceViewerInternal::ColumnType)
-						.OnSort(this, &SRenderResourceViewerWidget::OnColumnSortModeChanged)
-
-						+ SHeaderRow::Column(RenderResourceViewerInternal::ColumnSize)
-						.FillWidth(0.1f)
-						.DefaultLabel(LOCTEXT("AllocationSizeColumn", "Size"))
-						.SortMode(this, &SRenderResourceViewerWidget::GetColumnSortMode, RenderResourceViewerInternal::ColumnSize)
-						.OnSort(this, &SRenderResourceViewerWidget::OnColumnSortModeChanged)
-
-						+ SHeaderRow::Column(RenderResourceViewerInternal::ColumnFlags)
-						.FillWidth(0.2f)
-						.DefaultLabel(LOCTEXT("FlagsColumn", "Flags"))
-						.SortMode(this, &SRenderResourceViewerWidget::GetColumnSortMode, RenderResourceViewerInternal::ColumnFlags)
-						.OnSort(this, &SRenderResourceViewerWidget::OnColumnSortModeChanged)
+						RENDER_RESOURCE_VIEWER_ADD_COLUMN(RenderResourceViewerInternal::ColumnName, 0.3f, LOCTEXT("NameColumn", "Resource Name"))
+						RENDER_RESOURCE_VIEWER_ADD_COLUMN(RenderResourceViewerInternal::ColumnType, 0.1f, LOCTEXT("TypeColumn", "Type"))
+						RENDER_RESOURCE_VIEWER_ADD_COLUMN(RenderResourceViewerInternal::ColumnSize, 0.1f, LOCTEXT("SizeColumn", "Size"))
+						RENDER_RESOURCE_VIEWER_ADD_COLUMN(RenderResourceViewerInternal::ColumnFlags, 0.15f, LOCTEXT("FlagsColumn", "Flags"))
+						RENDER_RESOURCE_VIEWER_ADD_COLUMN(RenderResourceViewerInternal::ColumnOwner, 0.4f, LOCTEXT("OwnerColumn", "Owner"))
 					)
 				]
 			]
@@ -254,7 +249,7 @@ void SRenderResourceViewerWidget::RefreshNodes()
 	TotalResourceSize = 0;
 	for (const TSharedPtr<FRHIResourceStats>& Info : RHIResources)
 	{
-		bool bContainsFilterText = FilterText.IsEmpty() || Info->Name.Contains(FilterText.ToString());
+		bool bContainsFilterText = FilterText.IsEmpty() || Info->Name.ToString().Contains(FilterText.ToString()) || Info->OwnerName.ToString().Contains(FilterText.ToString());
 		// Note bMarkedForDelete resources are excluded from display
 		bool bContainsFilterFlags = (bShowTransient && Info->bTransient) || (bShowStreaming && Info->bStreaming)
 									|| (bShowRT && Info->bRenderTarget) || (bShowDS && Info->bDepthStencil) || (bShowUAV && Info->bUnorderedAccessView) || (bShowRTAS && Info->bRayTracingAccelerationStructure)
@@ -273,8 +268,8 @@ void SRenderResourceViewerWidget::RefreshNodes()
 		if (SortByColumn == RenderResourceViewerInternal::ColumnName)
 		{
 			ResourceInfos.Sort([this](const TSharedPtr<FRHIResourceStats>& A, const TSharedPtr<FRHIResourceStats>& B) {
-				if (SortMode == EColumnSortMode::Ascending) return (A->Name < B->Name);
-				else if (SortMode == EColumnSortMode::Descending) return (A->Name >= B->Name);
+				if (SortMode == EColumnSortMode::Ascending) return (A->Name.LexicalLess(B->Name));
+				else if (SortMode == EColumnSortMode::Descending) return (!A->Name.LexicalLess(B->Name));
 				return true;
 			});
 		}
@@ -302,6 +297,14 @@ void SRenderResourceViewerWidget::RefreshNodes()
 				return true;
 			});
 		}
+		else if (SortByColumn == RenderResourceViewerInternal::ColumnOwner)
+		{
+			ResourceInfos.Sort([this](const TSharedPtr<FRHIResourceStats>& A, const TSharedPtr<FRHIResourceStats>& B) {
+				if (SortMode == EColumnSortMode::Ascending) return (A->OwnerName.LexicalLess(B->OwnerName));
+				else if (SortMode == EColumnSortMode::Descending) return (!A->OwnerName.LexicalLess(B->OwnerName));
+				return true;
+			});
+		}
 	}
 
 	ResourceListView->RebuildList();
diff --git a/Engine/Source/Runtime/Engine/Classes/Engine/SkinnedAsset.h b/Engine/Source/Runtime/Engine/Classes/Engine/SkinnedAsset.h
--- a/Engine/Source/Runtime/Engine/Classes/Engine/SkinnedAsset.h
+++ b/Engine/Source/Runtime/Engine/Classes/Engine/SkinnedAsset.h
@@ -239,6 +239,9 @@ public:
 		
 	TArray<FSkinnedAssetVertexFactoryTypesPerMaterialData, TInlineAllocator<4>> GetVertexFactoryTypesPerMaterialIndex(int32 MinLODIndex, bool bCPUSkin, ERHIFeatureLevel::Type FeatureLevel);
 
+	/** Helper function for resource tracking, construct a string using the skinned asset's path name and LOD index . */
+	static FString GetLODPathName(const USkinnedAsset* Mesh, int32 LODIndex);
+
 #if WITH_EDITOR
 	/** IInterface_AsyncCompilation begin*/
 	virtual bool IsCompiling() const override;
diff --git a/Engine/Source/Runtime/Engine/Private/Animation/SkinWeightProfile.cpp b/Engine/Source/Runtime/Engine/Private/Animation/SkinWeightProfile.cpp
--- a/Engine/Source/Runtime/Engine/Private/Animation/SkinWeightProfile.cpp
+++ b/Engine/Source/Runtime/Engine/Private/Animation/SkinWeightProfile.cpp
@@ -371,7 +371,9 @@ void FSkinWeightProfilesData::SetDynamicDefaultSkinWeightProfile(USkeletalMesh*
 								DefaultOverrideSkinWeightBuffer = OverrideBuffer;
 								bDefaultOverriden = true;
 								DefaultProfileName = Profiles[DefaultProfileIndex].Name;
-
+								
+								const FName OwnerName(USkinnedAsset::GetLODPathName(Mesh, LODIndex));
+								OverrideBuffer->SetOwnerName(OwnerName);
 								OverrideBuffer->BeginInitResources();
 							}
 						}
@@ -744,6 +746,10 @@ void FSkinWeightProfilesData::InitialiseProfileBuffer(const FName& ProfileName)
 				ProfilePtr->ApplyOverrides(OverrideBuffer, BaseBufferData, BaseBuffer->GetNumVertices());
 			}
 
+#if RHI_ENABLE_RESOURCE_INFO
+			const FName OwnerName = FName(ProfileName.ToString() + TEXT("_FSkinWeightProfilesData"));
+			OverrideBuffer->SetOwnerName(OwnerName);
+#endif
 			OverrideBuffer->BeginInitResources();
 		}	
 	}
diff --git a/Engine/Source/Runtime/Engine/Private/Components/SkinnedMeshComponent.cpp b/Engine/Source/Runtime/Engine/Private/Components/SkinnedMeshComponent.cpp
--- a/Engine/Source/Runtime/Engine/Private/Components/SkinnedMeshComponent.cpp
+++ b/Engine/Source/Runtime/Engine/Private/Components/SkinnedMeshComponent.cpp
@@ -4300,7 +4300,11 @@ void USkinnedMeshComponent::SetSkinWeightOverride(int32 LODIndex, const TArray<F
 			Info.OverrideSkinWeights->SetMaxBoneInfluences(NumBoneInfluences);
 			Info.OverrideSkinWeights->SetUse16BitBoneIndex(bUse16BitBoneIndex);
 
-			const FReferenceSkeleton& RefSkel = GetSkinnedAsset()->GetRefSkeleton();
+			const USkinnedAsset* Mesh = GetSkinnedAsset();
+			const FName OwnerName(USkinnedAsset::GetLODPathName(Mesh, LODIndex));
+			Info.OverrideSkinWeights->SetOwnerName(OwnerName);
+
+			const FReferenceSkeleton& RefSkel = Mesh->GetRefSkeleton();
 			TArray<FSkinWeightInfo> GPUWeights;
 			CreateSkinWeightsArray(SkinWeights, LODData, GPUWeights, RefSkel);
 			*(Info.OverrideSkinWeights) = GPUWeights;
diff --git a/Engine/Source/Runtime/Engine/Private/RawIndexBuffer.cpp b/Engine/Source/Runtime/Engine/Private/RawIndexBuffer.cpp
--- a/Engine/Source/Runtime/Engine/Private/RawIndexBuffer.cpp
+++ b/Engine/Source/Runtime/Engine/Private/RawIndexBuffer.cpp
@@ -334,17 +334,21 @@ FBufferRHIRef FRawStaticIndexBuffer::CreateRHIBuffer_Internal()
 		const EBufferUsageFlags BufferFlags = EBufferUsageFlags::Static | (bSRV ? EBufferUsageFlags::ShaderResource : EBufferUsageFlags::None);
 
 		// Create the index buffer.
+		FBufferRHIRef Ret;
 		FRHIResourceCreateInfo CreateInfo(Is32Bit() ? TEXT("FRawStaticIndexBuffer32") : TEXT("FRawStaticIndexBuffer16"), &IndexStorage);
 		CreateInfo.bWithoutNativeResource = !SizeInBytes;
 		if (bRenderThread)
 		{
-			return RHICreateIndexBuffer(IndexStride, SizeInBytes, BufferFlags, CreateInfo);
+			Ret = RHICreateIndexBuffer(IndexStride, SizeInBytes, BufferFlags, CreateInfo);
 		}
 		else
 		{
 			FRHIAsyncCommandList CommandList;
-			return CommandList->CreateBuffer(SizeInBytes, BufferFlags | EBufferUsageFlags::IndexBuffer, IndexStride, ERHIAccess::SRVMask, CreateInfo);
+			Ret = CommandList->CreateBuffer(SizeInBytes, BufferFlags | EBufferUsageFlags::IndexBuffer, IndexStride, ERHIAccess::SRVMask, CreateInfo);
 		}
+
+		Ret->SetOwnerName(GetOwnerName());
+		return Ret;
 	}
 	return nullptr;
 }
diff --git a/Engine/Source/Runtime/Engine/Private/Rendering/MorphTargetVertexInfoBuffers.cpp b/Engine/Source/Runtime/Engine/Private/Rendering/MorphTargetVertexInfoBuffers.cpp
--- a/Engine/Source/Runtime/Engine/Private/Rendering/MorphTargetVertexInfoBuffers.cpp
+++ b/Engine/Source/Runtime/Engine/Private/Rendering/MorphTargetVertexInfoBuffers.cpp
@@ -17,6 +17,7 @@ void FMorphTargetVertexInfoBuffers::InitRHI()
 	const uint32 BufferSize = MorphData.Num() * sizeof(uint32);
 	FRHIResourceCreateInfo CreateInfo(TEXT("MorphData"));
 	MorphDataBuffer = RHICreateStructuredBuffer(sizeof(uint32), BufferSize, BUF_Static | BUF_ByteAddressBuffer | BUF_ShaderResource, ERHIAccess::SRVMask, CreateInfo);
+	MorphDataBuffer->SetOwnerName(GetOwnerName());
 	
 	void* BufferPtr = RHILockBuffer(MorphDataBuffer, 0, BufferSize, RLM_WriteOnly);
 	FMemory::ParallelMemcpy(BufferPtr, MorphData.GetData(), BufferSize, EMemcpyCachePolicy::StoreUncached);
diff --git a/Engine/Source/Runtime/Engine/Private/Rendering/MultiSizeIndexContainer.cpp b/Engine/Source/Runtime/Engine/Private/Rendering/MultiSizeIndexContainer.cpp
--- a/Engine/Source/Runtime/Engine/Private/Rendering/MultiSizeIndexContainer.cpp
+++ b/Engine/Source/Runtime/Engine/Private/Rendering/MultiSizeIndexContainer.cpp
@@ -11,6 +11,15 @@ FMultiSizeIndexContainer::~FMultiSizeIndexContainer()
 	}
 }
 
+void FMultiSizeIndexContainer::SetOwnerName(const FName& OwnerName)
+{
+	check(IsInGameThread());
+	if (IndexBuffer)
+	{
+		IndexBuffer->SetOwnerName(OwnerName);
+	}
+}
+
 /**
 * Initialize the index buffer's render resources.
 */
diff --git a/Engine/Source/Runtime/Engine/Private/Rendering/SkinWeightVertexBuffer.cpp b/Engine/Source/Runtime/Engine/Private/Rendering/SkinWeightVertexBuffer.cpp
--- a/Engine/Source/Runtime/Engine/Private/Rendering/SkinWeightVertexBuffer.cpp
+++ b/Engine/Source/Runtime/Engine/Private/Rendering/SkinWeightVertexBuffer.cpp
@@ -758,6 +758,13 @@ void FSkinWeightVertexBuffer::ResetVertexBoneWeights(uint32 VertexIndex)
 	DataVertexBuffer.ResetVertexBoneWeights(VertexWeightOffset, VertexInfluenceCount);
 }
 
+void FSkinWeightVertexBuffer::SetOwnerName(const FName& OwnerName)
+{
+	check(IsInGameThread());
+	LookupVertexBuffer.SetOwnerName(OwnerName);
+	DataVertexBuffer.SetOwnerName(OwnerName);
+}
+
 void FSkinWeightVertexBuffer::BeginInitResources()
 {
 	BeginInitResource(&LookupVertexBuffer);
diff --git a/Engine/Source/Runtime/Engine/Private/SkeletalMeshLODRenderData.cpp b/Engine/Source/Runtime/Engine/Private/SkeletalMeshLODRenderData.cpp
--- a/Engine/Source/Runtime/Engine/Private/SkeletalMeshLODRenderData.cpp
+++ b/Engine/Source/Runtime/Engine/Private/SkeletalMeshLODRenderData.cpp
@@ -213,27 +213,35 @@ void FSkeletalMeshLODRenderData::InitMorphResources()
 
 void FSkeletalMeshLODRenderData::InitResources(bool bNeedsVertexColors, int32 LODIndex, TArray<UMorphTarget*>& InMorphTargets, USkinnedAsset* Owner)
 {
+	const FName OwnerName(USkinnedAsset::GetLODPathName(Owner, LODIndex));
+
 	if (bStreamedDataInlined)
 	{
 		IncrementMemoryStats(bNeedsVertexColors);
 	}
 
+	MultiSizeIndexContainer.SetOwnerName(OwnerName);
 	MultiSizeIndexContainer.InitResources();
 
+	StaticVertexBuffers.PositionVertexBuffer.SetOwnerName(OwnerName);
 	BeginInitResource(&StaticVertexBuffers.PositionVertexBuffer);
+	StaticVertexBuffers.StaticMeshVertexBuffer.SetOwnerName(OwnerName);
 	BeginInitResource(&StaticVertexBuffers.StaticMeshVertexBuffer);
 
+	SkinWeightVertexBuffer.SetOwnerName(OwnerName);
 	SkinWeightVertexBuffer.BeginInitResources();
 
 	if (bNeedsVertexColors)
 	{
 		// Only init the color buffer if the mesh has vertex colors
+		StaticVertexBuffers.ColorVertexBuffer.SetOwnerName(OwnerName);
 		BeginInitResource(&StaticVertexBuffers.ColorVertexBuffer);
 	}
 
 	if (ClothVertexBuffer.GetNumVertices() > 0)
 	{
 		// Only init the clothing buffer if the mesh has clothing data
+		ClothVertexBuffer.SetOwnerName(OwnerName);
 		BeginInitResource(&ClothVertexBuffer);
 	}
 
@@ -254,6 +262,7 @@ void FSkeletalMeshLODRenderData::InitResources(bool bNeedsVertexColors, int32 LO
 			{
 				// No need to discard CPU data in cooked builds as bNeedsCPUAccess is false (see FDuplicatedVerticesBuffer constructor), 
 				// so it'd be auto-discarded after the RHI has copied the resource data. Keep CPU data when in the editor for geometry operations.
+				RenderSection.DuplicatedVerticesBuffer.SetOwnerName(OwnerName);
 				BeginInitResource(&RenderSection.DuplicatedVerticesBuffer);
 			}
 		}
@@ -269,6 +278,7 @@ void FSkeletalMeshLODRenderData::InitResources(bool bNeedsVertexColors, int32 LO
 	
 	if (!MorphTargetVertexInfoBuffers.IsRHIIntialized() && MorphTargetVertexInfoBuffers.IsMorphCPUDataValid() && MorphTargetVertexInfoBuffers.NumTotalBatches > 0)
 	{
+		MorphTargetVertexInfoBuffers.SetOwnerName(OwnerName);
 		BeginInitResource(&MorphTargetVertexInfoBuffers);
 	}
 
@@ -277,6 +287,7 @@ void FSkeletalMeshLODRenderData::InitResources(bool bNeedsVertexColors, int32 LO
 	{
 		if (SourceRayTracingGeometry.RawData.Num() > 0)
 		{
+			SourceRayTracingGeometry.SetOwnerName(OwnerName);
 			BeginInitResource(&SourceRayTracingGeometry);
 		}
 	}
diff --git a/Engine/Source/Runtime/Engine/Private/SkinnedAsset.cpp b/Engine/Source/Runtime/Engine/Private/SkinnedAsset.cpp
--- a/Engine/Source/Runtime/Engine/Private/SkinnedAsset.cpp
+++ b/Engine/Source/Runtime/Engine/Private/SkinnedAsset.cpp
@@ -377,3 +377,12 @@ TArray<FSkeletalMaterial>& USkinnedAsset::GetSkeletalMaterialDummyArray()
 	static TArray<FSkeletalMaterial> Dummy;
 	return Dummy;
 }
+
+FString USkinnedAsset::GetLODPathName(const USkinnedAsset* Mesh, int32 LODIndex)
+{
+#if RHI_ENABLE_RESOURCE_INFO
+	return FString::Printf(TEXT("%s [LOD%d]"), Mesh ? *Mesh->GetPathName() : TEXT("UnknownSkinnedAsset"), LODIndex);
+#else
+	return TEXT("");
+#endif
+}
\ No newline at end of file
diff --git a/Engine/Source/Runtime/Engine/Private/StaticMesh.cpp b/Engine/Source/Runtime/Engine/Private/StaticMesh.cpp
--- a/Engine/Source/Runtime/Engine/Private/StaticMesh.cpp
+++ b/Engine/Source/Runtime/Engine/Private/StaticMesh.cpp
@@ -1318,43 +1318,54 @@ void FStaticMeshLODResources::UpdateVertexMemoryStats() const
 #endif
 }
 
-void FStaticMeshLODResources::InitResources(UStaticMesh* Parent)
+void FStaticMeshLODResources::InitResources(UStaticMesh* Parent, int32 LODIndex)
 {
+	const FName OwnerName(FString::Printf(TEXT("%s [LOD%d]"), Parent ? *Parent->GetPathName() : TEXT("UnknownStaticMesh"), LODIndex));
+
 	if (bBuffersInlined)
 	{
 		UpdateIndexMemoryStats<true>();
 	}
 
+	IndexBuffer.SetOwnerName(OwnerName);
 	BeginInitResource(&IndexBuffer);
 	if(bHasWireframeIndices)
 	{
+		AdditionalIndexBuffers->WireframeIndexBuffer.SetOwnerName(OwnerName);
 		BeginInitResource(&AdditionalIndexBuffers->WireframeIndexBuffer);
 	}
+	VertexBuffers.StaticMeshVertexBuffer.SetOwnerName(OwnerName);
 	BeginInitResource(&VertexBuffers.StaticMeshVertexBuffer);
+	VertexBuffers.PositionVertexBuffer.SetOwnerName(OwnerName);
 	BeginInitResource(&VertexBuffers.PositionVertexBuffer);
 	if(bHasColorVertexData)
 	{
+		VertexBuffers.ColorVertexBuffer.SetOwnerName(OwnerName);
 		BeginInitResource(&VertexBuffers.ColorVertexBuffer);
 	}
 
 	if (bHasReversedIndices)
 	{
+		AdditionalIndexBuffers->ReversedIndexBuffer.SetOwnerName(OwnerName);
 		BeginInitResource(&AdditionalIndexBuffers->ReversedIndexBuffer);
 	}
 
 	if (bHasDepthOnlyIndices)
 	{
+		DepthOnlyIndexBuffer.SetOwnerName(OwnerName);
 		BeginInitResource(&DepthOnlyIndexBuffer);
 	}
 
 	if (bHasReversedDepthOnlyIndices)
 	{
+		AdditionalIndexBuffers->ReversedDepthOnlyIndexBuffer.SetOwnerName(OwnerName);
 		BeginInitResource(&AdditionalIndexBuffers->ReversedDepthOnlyIndexBuffer);
 	}
 
 	if (Parent->bSupportGpuUniformlyDistributedSampling && Parent->bSupportUniformlyDistributedSampling && (AreaWeightedSampler.GetNumEntries() > 0))
 	{
 		AreaWeightedSectionSamplersBuffer.Init(&AreaWeightedSectionSamplers);
+		AreaWeightedSectionSamplersBuffer.SetOwnerName(OwnerName);
 		BeginInitResource(&AreaWeightedSectionSamplersBuffer);
 	}
 
@@ -1882,7 +1893,7 @@ void FStaticMeshRenderData::InitResources(ERHIFeatureLevel::Type InFeatureLevel,
 		// Skip LODs that have their render data stripped
 		if (LODResources[LODIndex].VertexBuffers.StaticMeshVertexBuffer.GetNumVertices() > 0)
 		{
-			LODResources[LODIndex].InitResources(Owner);
+			LODResources[LODIndex].InitResources(Owner, LODIndex);
 			LODVertexFactories[LODIndex].InitResources(LODResources[LODIndex], LODIndex, Owner);
 		}
 	}
diff --git a/Engine/Source/Runtime/Engine/Public/RawIndexBuffer.h b/Engine/Source/Runtime/Engine/Public/RawIndexBuffer.h
--- a/Engine/Source/Runtime/Engine/Public/RawIndexBuffer.h
+++ b/Engine/Source/Runtime/Engine/Public/RawIndexBuffer.h
@@ -572,9 +572,10 @@ private:
 			else
 			{
 				FRHIAsyncCommandList CommandList;
-				return CommandList->CreateBuffer(Size, Flags | EBufferUsageFlags::IndexBuffer, sizeof(INDEX_TYPE), ERHIAccess::SRVMask, CreateInfo);
+				Ret = CommandList->CreateBuffer(Size, Flags | EBufferUsageFlags::IndexBuffer, sizeof(INDEX_TYPE), ERHIAccess::SRVMask, CreateInfo);
 			}
 
+			Ret->SetOwnerName(GetOwnerName());
 			return Ret;
 		}
 		return nullptr;
diff --git a/Engine/Source/Runtime/Engine/Public/Rendering/MultiSizeIndexContainer.h b/Engine/Source/Runtime/Engine/Public/Rendering/MultiSizeIndexContainer.h
--- a/Engine/Source/Runtime/Engine/Public/Rendering/MultiSizeIndexContainer.h
+++ b/Engine/Source/Runtime/Engine/Public/Rendering/MultiSizeIndexContainer.h
@@ -25,6 +25,11 @@ public:
 
 	ENGINE_API ~FMultiSizeIndexContainer();
 
+	/**
+	* SetOwnerName should be called before BeginInitResources for the owner name to be successfully tracked.
+	*/
+	ENGINE_API void SetOwnerName(const FName& OwnerName);
+
 	/**
 	* Initialize the index buffer's render resources.
 	*/
diff --git a/Engine/Source/Runtime/Engine/Public/Rendering/SkinWeightVertexBuffer.h b/Engine/Source/Runtime/Engine/Public/Rendering/SkinWeightVertexBuffer.h
--- a/Engine/Source/Runtime/Engine/Public/Rendering/SkinWeightVertexBuffer.h
+++ b/Engine/Source/Runtime/Engine/Public/Rendering/SkinWeightVertexBuffer.h
@@ -514,6 +514,8 @@ public:
 	ENGINE_API void SetBoneWeight(uint32 VertexIndex, uint32 InfluenceIndex, uint16 BoneWeight);
 	ENGINE_API void ResetVertexBoneWeights(uint32 VertexIndex);
 
+	/** SetOwnerName should be called before BeginInitResources for the owner name to be successfully tracked. */
+	ENGINE_API void SetOwnerName(const FName& OwnerName);
 	ENGINE_API void BeginInitResources();
 	ENGINE_API void BeginReleaseResources();
 	ENGINE_API void ReleaseResources();
diff --git a/Engine/Source/Runtime/Engine/Public/StaticMeshResources.h b/Engine/Source/Runtime/Engine/Public/StaticMeshResources.h
--- a/Engine/Source/Runtime/Engine/Public/StaticMeshResources.h
+++ b/Engine/Source/Runtime/Engine/Public/StaticMeshResources.h
@@ -455,7 +455,7 @@ public:
 	}
 
 	/** Initializes all rendering resources. */
-	void InitResources(UStaticMesh* Parent);
+	void InitResources(UStaticMesh* Parent, int32 LODIndex);
 
 	/** Releases all rendering resources. */
 	void ReleaseResources();
diff --git a/Engine/Source/Runtime/RHI/Private/RHI.cpp b/Engine/Source/Runtime/RHI/Private/RHI.cpp
--- a/Engine/Source/Runtime/RHI/Private/RHI.cpp
+++ b/Engine/Source/Runtime/RHI/Private/RHI.cpp
@@ -668,6 +668,22 @@ bool operator==(const FBlendStateInitializerRHI::FRenderTarget& A, const FBlendS
 	return bSame;
 }
 
+FName FRHIResource::GetOwnerName() const
+{
+#if RHI_ENABLE_RESOURCE_INFO
+	return OwnerName;
+#else
+	return NAME_None;
+#endif
+}
+
+void FRHIResource::SetOwnerName(const FName& InOwnerName)
+{
+#if RHI_ENABLE_RESOURCE_INFO
+	OwnerName = InOwnerName;
+#endif
+}
+
 #if RHI_ENABLE_RESOURCE_INFO
 
 static FCriticalSection GRHIResourceTrackingCriticalSection;
@@ -1065,12 +1081,12 @@ void RHIGetTrackedResourceStats(TArray<TSharedPtr<FRHIResourceStats>>& OutResour
 	OutResourceStats.SetNum(Resources.Num());
 	ParallelFor(Resources.Num(), [&](int32 Index)
 	{
+		const FRHIResource* Resource = Resources[Index].Resource;
 		const FRHIResourceInfo& ResourceInfo = Resources[Index].ResourceInfo;
-		FString ResourceName = ResourceInfo.Name.ToString();
 		const TCHAR* ResourceType = StringFromRHIResourceType(ResourceInfo.Type);
 		const int64 SizeInBytes = ResourceInfo.VRamAllocation.AllocationSize;
 		RHIInternal::FResourceFlags Flags = GetResourceFlagsInternal(Resources[Index]);
-		OutResourceStats[Index] = MakeShared<FRHIResourceStats>(ResourceName, ResourceType, Flags.GetString(), SizeInBytes,
+		OutResourceStats[Index] = MakeShared<FRHIResourceStats>(ResourceInfo.Name, Resource->GetOwnerName(), ResourceType, Flags.GetString(), SizeInBytes,
 									Flags.bMarkedForDelete, Flags.bTransient, Flags.bStreaming, Flags.bRT, Flags.bDS, Flags.bUAV, Flags.bRTAS, Flags.bHasFlags);
 	});
 }
diff --git a/Engine/Source/Runtime/RHI/Public/RHI.h b/Engine/Source/Runtime/RHI/Public/RHI.h
--- a/Engine/Source/Runtime/RHI/Public/RHI.h
+++ b/Engine/Source/Runtime/RHI/Public/RHI.h
@@ -123,7 +123,8 @@ RHI_API void RHIDumpResourceMemoryToCSV();
 
 struct FRHIResourceStats
 {
-	FString Name;
+	FName Name;
+	FName OwnerName;
 	FString Type;
 	FString Flags;
 	uint64	SizeInBytes = 0;
@@ -136,9 +137,10 @@ struct FRHIResourceStats
 	bool	bRayTracingAccelerationStructure = false;
 	bool	bHasFlags = false;
 
-	FRHIResourceStats(const FString& InName, const FString& InType, const FString& InFlags, const uint64& InSizeInBytes, 
+	FRHIResourceStats(const FName& InName, const FName& InOwnerName, const FString& InType, const FString& InFlags, const uint64& InSizeInBytes,
 						bool bInMarkedForDelete, bool bInTransient, bool bInStreaming, bool bInRT, bool bInDS, bool bInUAV, bool bInRTAS, bool bInHasFlags)
 		: Name(InName)
+		, OwnerName(InOwnerName)
 		, Type(InType)
 		, Flags(InFlags)
 		, SizeInBytes(InSizeInBytes)
diff --git a/Engine/Source/Runtime/RHI/Public/RHIResources.h b/Engine/Source/Runtime/RHI/Public/RHIResources.h
--- a/Engine/Source/Runtime/RHI/Public/RHIResources.h
+++ b/Engine/Source/Runtime/RHI/Public/RHIResources.h
@@ -128,6 +128,9 @@ public:
 
 	inline ERHIResourceType GetType() const { return ResourceType; }
 
+	FName GetOwnerName() const;
+	void SetOwnerName(const FName& InOwnerName);
+
 #if RHI_ENABLE_RESOURCE_INFO
 	// Get resource info if available.
 	// Should return true if the ResourceInfo was filled with data.
@@ -221,6 +224,7 @@ private:
 	uint8 bCommitted : 1;
 #if RHI_ENABLE_RESOURCE_INFO
 	uint8 bBeingTracked : 1;
+	FName OwnerName;
 #endif
 
 	static UE::TConsumeAllMpmcQueue<FRHIResource*> PendingDeletes;
diff --git a/Engine/Source/Runtime/RenderCore/Private/RenderResource.cpp b/Engine/Source/Runtime/RenderCore/Private/RenderResource.cpp
--- a/Engine/Source/Runtime/RenderCore/Private/RenderResource.cpp
+++ b/Engine/Source/Runtime/RenderCore/Private/RenderResource.cpp
@@ -325,6 +325,23 @@ FRenderResource::~FRenderResource()
 	}
 }
 
+void FRenderResource::SetOwnerName(const FName& InOwnerName)
+{
+#if RHI_ENABLE_RESOURCE_INFO
+	check(IsInGameThread());
+	OwnerName = InOwnerName;
+#endif
+}
+
+FName FRenderResource::GetOwnerName() const
+{
+#if RHI_ENABLE_RESOURCE_INFO
+	return OwnerName;
+#else
+	return NAME_None;
+#endif
+}
+
 void BeginInitResource(FRenderResource* Resource)
 {
 	LLM_SCOPE(ELLMTag::SceneRender);
diff --git a/Engine/Source/Runtime/RenderCore/Public/RenderResource.h b/Engine/Source/Runtime/RenderCore/Public/RenderResource.h
--- a/Engine/Source/Runtime/RenderCore/Public/RenderResource.h
+++ b/Engine/Source/Runtime/RenderCore/Public/RenderResource.h
@@ -145,8 +145,15 @@ public:
 
 	int32 GetListIndex() const { return ListIndex; }
 
+	/** SetOwnerName should be called before BeginInitResource for the owner name to be successfully tracked. */
+	void SetOwnerName(const FName& InOwnerName);
+	FName GetOwnerName() const;
+
 private:
 	int32 ListIndex;
+#if RHI_ENABLE_RESOURCE_INFO
+	FName OwnerName;
+#endif
 
 protected:
 	// This is used during mobile editor preview refactor, this will eventually be replaced with a parameter to InitRHI() etc..
@@ -169,11 +176,14 @@ protected:
 			if (bRenderThread)
 			{
 				Buffer = RHICreateVertexBuffer(SizeInBytes, InBufferUsageFlags, CreateInfo);
+				Buffer->SetOwnerName(GetOwnerName());
 			}
 			else
 			{
 				FRHIAsyncCommandList CommandList;
-				return CommandList->CreateBuffer(SizeInBytes, InBufferUsageFlags | EBufferUsageFlags::VertexBuffer, 0, ERHIAccess::SRVMask, CreateInfo);
+				Buffer = CommandList->CreateBuffer(SizeInBytes, InBufferUsageFlags | EBufferUsageFlags::VertexBuffer, 0, ERHIAccess::SRVMask, CreateInfo);
+				Buffer->SetOwnerName(GetOwnerName());
+				return Buffer;
 			}
 		}
 
