Fix render resource viewer None checkbox filter

diff --git a/Engine/Source/Editor/RenderResourceViewer/Private/SRenderResourceViewer.cpp b/Engine/Source/Editor/RenderResourceViewer/Private/SRenderResourceViewer.cpp
--- a/Engine/Source/Editor/RenderResourceViewer/Private/SRenderResourceViewer.cpp
+++ b/Engine/Source/Editor/RenderResourceViewer/Private/SRenderResourceViewer.cpp
@@ -144,7 +144,7 @@ void SRenderResourceViewerWidget::Construct(const FArguments& InArgs, const TSha
 					RENDER_RESOURCE_VIEWER_ADD_CHECKBOX(bShowDS, LOCTEXT("DSText", "DS"), LOCTEXT("DSTooltip", "Depth Stencil"), SRenderResourceViewerWidget::OnDSCheckboxChanged)
 					RENDER_RESOURCE_VIEWER_ADD_CHECKBOX(bShowUAV, LOCTEXT("UAVText", "UAV"), LOCTEXT("UAVTooltip", "Unordered Access View"), SRenderResourceViewerWidget::OnUAVCheckboxChanged)
 					RENDER_RESOURCE_VIEWER_ADD_CHECKBOX(bShowRTAS, LOCTEXT("RTASText", "RTAS"), LOCTEXT("RTASTooltip", "Ray Tracing Acceleration Structure"), SRenderResourceViewerWidget::OnRTASCheckboxChanged)
-					RENDER_RESOURCE_VIEWER_ADD_CHECKBOX(bShowRTAS, LOCTEXT("NoneText", "None"), LOCTEXT("NoneTooltip", "Resource with no flags set"), SRenderResourceViewerWidget::OnNoneCheckboxChanged)
+					RENDER_RESOURCE_VIEWER_ADD_CHECKBOX(bShowNone, LOCTEXT("NoneText", "None"), LOCTEXT("NoneTooltip", "Resource with no flags set"), SRenderResourceViewerWidget::OnNoneCheckboxChanged)
 
 					// Refresh button to update the resource list
 					+ SHorizontalBox::Slot()
