Integrate RenderResourceViewer from //UE5/Main
Part4 - Add support for Resident flag filtering. Fixes render resources with only Resident flag don't show up.

diff --git a/Engine/Source/Editor/RenderResourceViewer/Private/SRenderResourceViewer.cpp b/Engine/Source/Editor/RenderResourceViewer/Private/SRenderResourceViewer.cpp
--- a/Engine/Source/Editor/RenderResourceViewer/Private/SRenderResourceViewer.cpp
+++ b/Engine/Source/Editor/RenderResourceViewer/Private/SRenderResourceViewer.cpp
@@ -138,6 +138,7 @@ void SRenderResourceViewerWidget::Construct(const FArguments& InArgs, const TSha
 					]
 
 					// A row of checkboxes to filter resource list by flags
+					RENDER_RESOURCE_VIEWER_ADD_CHECKBOX(bShowResident, LOCTEXT("ResidentText", "Resident"), LOCTEXT("ResidentTooltip", "Resident"), SRenderResourceViewerWidget::OnResidentCheckboxChanged)
 					RENDER_RESOURCE_VIEWER_ADD_CHECKBOX(bShowTransient, LOCTEXT("TransientText", "Transient"), LOCTEXT("TransientTooltip", "Transient"), SRenderResourceViewerWidget::OnTransientCheckboxChanged)
 					RENDER_RESOURCE_VIEWER_ADD_CHECKBOX(bShowStreaming, LOCTEXT("StreamingText", "Streaming"), LOCTEXT("StreamingTooltip", "Streaming"), SRenderResourceViewerWidget::OnStreamingCheckboxChanged)
 					RENDER_RESOURCE_VIEWER_ADD_CHECKBOX(bShowRT, LOCTEXT("RTText", "RT"), LOCTEXT("RTTooltip", "Render Target"), SRenderResourceViewerWidget::OnRTCheckboxChanged)
@@ -251,7 +252,7 @@ void SRenderResourceViewerWidget::RefreshNodes()
 	{
 		bool bContainsFilterText = FilterText.IsEmpty() || Info->Name.ToString().Contains(FilterText.ToString()) || Info->OwnerName.ToString().Contains(FilterText.ToString());
 		// Note bMarkedForDelete resources are excluded from display
-		bool bContainsFilterFlags = (bShowTransient && Info->bTransient) || (bShowStreaming && Info->bStreaming)
+		bool bContainsFilterFlags = (bShowResident && Info->bResident) || (bShowTransient && Info->bTransient) || (bShowStreaming && Info->bStreaming)
 									|| (bShowRT && Info->bRenderTarget) || (bShowDS && Info->bDepthStencil) || (bShowUAV && Info->bUnorderedAccessView) || (bShowRTAS && Info->bRayTracingAccelerationStructure)
 									|| (bShowNone && !Info->bHasFlags);
 		if (bContainsFilterText && bContainsFilterFlags)
diff --git a/Engine/Source/Editor/RenderResourceViewer/Private/SRenderResourceViewer.h b/Engine/Source/Editor/RenderResourceViewer/Private/SRenderResourceViewer.h
--- a/Engine/Source/Editor/RenderResourceViewer/Private/SRenderResourceViewer.h
+++ b/Engine/Source/Editor/RenderResourceViewer/Private/SRenderResourceViewer.h
@@ -31,6 +31,7 @@ private:
 	FText GetResourceCountText() const									{ return FText::AsNumber(TotalResourceCount); }
 	FText GetResourceSizeText() const;
 	void FilterTextBox_OnTextChanged(const FText& InFilterText)			{ FilterText = InFilterText; RefreshNodes(); }
+	void OnResidentCheckboxChanged(ECheckBoxState NewState)				{ bShowResident = (NewState == ECheckBoxState::Checked); RefreshNodes(); }
 	void OnTransientCheckboxChanged(ECheckBoxState NewState)			{ bShowTransient = (NewState == ECheckBoxState::Checked); RefreshNodes(); }
 	void OnStreamingCheckboxChanged(ECheckBoxState NewState)			{ bShowStreaming = (NewState == ECheckBoxState::Checked); RefreshNodes(); }
 	void OnRTCheckboxChanged(ECheckBoxState NewState)					{ bShowRT = (NewState == ECheckBoxState::Checked); RefreshNodes(); }
@@ -49,6 +50,7 @@ private:
 	uint64 TotalResourceCount = 0;
 	uint64 TotalResourceSize = 0;
 
+	bool bShowResident = true;			// Show resource with Resident flag set
 	bool bShowTransient = true;			// Show resource with Transient flag set
 	bool bShowStreaming = true;			// Show resource with Streaming flag set
 	bool bShowRT = true;				// Show resource with RenderTarget flag set
diff --git a/Engine/Source/Runtime/RHI/Private/RHI.cpp b/Engine/Source/Runtime/RHI/Private/RHI.cpp
--- a/Engine/Source/Runtime/RHI/Private/RHI.cpp
+++ b/Engine/Source/Runtime/RHI/Private/RHI.cpp
@@ -833,7 +833,7 @@ void RHIGetTrackedResourceStats(TArray<TSharedPtr<FRHIResourceStats>>& OutResour
 		const int64 SizeInBytes = ResourceInfo.VRamAllocation.AllocationSize;
 		RHIInternal::FResourceFlags Flags = GetResourceFlagsInternal(Resources[Index]);
 		OutResourceStats[Index] = MakeShared<FRHIResourceStats>(ResourceInfo.Name, Resource->GetOwnerName(), ResourceType, Flags.GetString(), SizeInBytes,
-									Flags.bMarkedForDelete, Flags.bTransient, Flags.bStreaming, Flags.bRT, Flags.bDS, Flags.bUAV, Flags.bRTAS, Flags.bHasFlags);
+									Flags.bResident, Flags.bMarkedForDelete, Flags.bTransient, Flags.bStreaming, Flags.bRT, Flags.bDS, Flags.bUAV, Flags.bRTAS, Flags.bHasFlags);
 	});
 }
 
diff --git a/Engine/Source/Runtime/RHI/Public/RHI.h b/Engine/Source/Runtime/RHI/Public/RHI.h
--- a/Engine/Source/Runtime/RHI/Public/RHI.h
+++ b/Engine/Source/Runtime/RHI/Public/RHI.h
@@ -127,6 +127,7 @@ struct FRHIResourceStats
 	FString Type;
 	FString Flags;
 	uint64	SizeInBytes = 0;
+	bool	bResident = false;
 	bool	bMarkedForDelete = false;
 	bool	bTransient = false;
 	bool	bStreaming = false;
@@ -137,12 +138,13 @@ struct FRHIResourceStats
 	bool	bHasFlags = false;
 
 	FRHIResourceStats(const FName& InName, const FName& InOwnerName, const FString& InType, const FString& InFlags, const uint64& InSizeInBytes,
-						bool bInMarkedForDelete, bool bInTransient, bool bInStreaming, bool bInRT, bool bInDS, bool bInUAV, bool bInRTAS, bool bInHasFlags)
+						bool bInResident, bool bInMarkedForDelete, bool bInTransient, bool bInStreaming, bool bInRT, bool bInDS, bool bInUAV, bool bInRTAS, bool bInHasFlags)
 		: Name(InName)
 		, OwnerName(InOwnerName)
 		, Type(InType)
 		, Flags(InFlags)
 		, SizeInBytes(InSizeInBytes)
+		, bResident(bInResident)
 		, bMarkedForDelete(bInMarkedForDelete)
 		, bTransient(bInTransient)
 		, bStreaming(bInStreaming)
