Hide MarkedForDelete resources from RenderResourceViewer

diff --git a/Engine/Source/Editor/RenderResourceViewer/Private/SRenderResourceViewer.cpp b/Engine/Source/Editor/RenderResourceViewer/Private/SRenderResourceViewer.cpp
--- a/Engine/Source/Editor/RenderResourceViewer/Private/SRenderResourceViewer.cpp
+++ b/Engine/Source/Editor/RenderResourceViewer/Private/SRenderResourceViewer.cpp
@@ -255,6 +255,11 @@ void SRenderResourceViewerWidget::RefreshNodes()
 		bool bContainsFilterFlags = (bShowResident && Info->bResident) || (bShowTransient && Info->bTransient) || (bShowStreaming && Info->bStreaming)
 									|| (bShowRT && Info->bRenderTarget) || (bShowDS && Info->bDepthStencil) || (bShowUAV && Info->bUnorderedAccessView) || (bShowRTAS && Info->bRayTracingAccelerationStructure)
 									|| (bShowNone && !Info->bHasFlags);
+		// bMarkedForDelete resources are excluded from display
+		if (Info->bMarkedForDelete)
+		{
+			bContainsFilterFlags = false;
+		}
 		if (bContainsFilterText && bContainsFilterFlags)
 		{
 			ResourceInfos.Add(Info);
