Fix Render Resource Viewer resource size changes when using filtering toggle and search by text functions

diff --git a/Engine/Source/Editor/RenderResourceViewer/Private/SRenderResourceViewer.cpp b/Engine/Source/Editor/RenderResourceViewer/Private/SRenderResourceViewer.cpp
--- a/Engine/Source/Editor/RenderResourceViewer/Private/SRenderResourceViewer.cpp
+++ b/Engine/Source/Editor/RenderResourceViewer/Private/SRenderResourceViewer.cpp
@@ -232,17 +232,20 @@ void SRenderResourceViewerWidget::Construct(const FArguments& InArgs, const TSha
 		]
 	];
 
-	RefreshNodes();
+	RefreshNodes(true);
 }
 
-void SRenderResourceViewerWidget::RefreshNodes()
+void SRenderResourceViewerWidget::RefreshNodes(bool bUpdateRHIResources)
 {
 	const int32 NumberOfResourcesToShow = -1;
 
 	// Retrieve the whole list of RHI tracked resources
-	TArray<TSharedPtr<FRHIResourceStats>> RHIResources;
 #if RHI_ENABLE_RESOURCE_INFO
-	RHIGetTrackedResourceStats(RHIResources);
+	if (bUpdateRHIResources)
+	{
+		RHIResources.Reset();
+		RHIGetTrackedResourceStats(RHIResources);
+	}
 #endif
 
 	// Filter the resource list by search keyword and flags
diff --git a/Engine/Source/Editor/RenderResourceViewer/Private/SRenderResourceViewer.h b/Engine/Source/Editor/RenderResourceViewer/Private/SRenderResourceViewer.h
--- a/Engine/Source/Editor/RenderResourceViewer/Private/SRenderResourceViewer.h
+++ b/Engine/Source/Editor/RenderResourceViewer/Private/SRenderResourceViewer.h
@@ -23,7 +23,7 @@ public:
 	~SRenderResourceViewerWidget() {}
 
 private:
-	void RefreshNodes();
+	void RefreshNodes(bool bUpdateRHIResources = false);
 	TSharedRef<ITableRow> HandleResourceGenerateRow(TSharedPtr<FRHIResourceStats> Item, const TSharedRef<STableViewBase>& OwnerTable);
 	void OnColumnSortModeChanged(const EColumnSortPriority::Type SortPriority, const FName& ColumnId, const EColumnSortMode::Type InSortMode);
 
@@ -39,8 +39,9 @@ private:
 	void OnUAVCheckboxChanged(ECheckBoxState NewState)					{ bShowUAV = (NewState == ECheckBoxState::Checked); RefreshNodes(); }
 	void OnRTASCheckboxChanged(ECheckBoxState NewState)					{ bShowRTAS = (NewState == ECheckBoxState::Checked); RefreshNodes(); }
 	void OnNoneCheckboxChanged(ECheckBoxState NewState)					{ bShowNone = (NewState == ECheckBoxState::Checked); RefreshNodes(); }
-	FReply OnRefreshButtonClicked()										{ RefreshNodes(); return FReply::Handled(); }
+	FReply OnRefreshButtonClicked()										{ RefreshNodes(true); return FReply::Handled(); }
 
+	TArray<TSharedPtr<FRHIResourceStats>> RHIResources;
 	TSharedPtr<SListView<TSharedPtr<FRHIResourceStats>>> ResourceListView;
 	TArray<TSharedPtr<FRHIResourceStats>> ResourceInfos;
 	FName SortByColumn;
