Improve render resource viewer flags tooltips

diff --git a/Engine/Source/Editor/RenderResourceViewer/Private/SRenderResourceViewer.cpp b/Engine/Source/Editor/RenderResourceViewer/Private/SRenderResourceViewer.cpp
--- a/Engine/Source/Editor/RenderResourceViewer/Private/SRenderResourceViewer.cpp
+++ b/Engine/Source/Editor/RenderResourceViewer/Private/SRenderResourceViewer.cpp
@@ -138,13 +138,13 @@ void SRenderResourceViewerWidget::Construct(const FArguments& InArgs, const TSha
 					]
 
 					// A row of checkboxes to filter resource list by flags
-					RENDER_RESOURCE_VIEWER_ADD_CHECKBOX(bShowResident, LOCTEXT("ResidentText", "Resident"), LOCTEXT("ResidentTooltip", "Resident"), SRenderResourceViewerWidget::OnResidentCheckboxChanged)
-					RENDER_RESOURCE_VIEWER_ADD_CHECKBOX(bShowTransient, LOCTEXT("TransientText", "Transient"), LOCTEXT("TransientTooltip", "Transient"), SRenderResourceViewerWidget::OnTransientCheckboxChanged)
-					RENDER_RESOURCE_VIEWER_ADD_CHECKBOX(bShowStreaming, LOCTEXT("StreamingText", "Streaming"), LOCTEXT("StreamingTooltip", "Streaming"), SRenderResourceViewerWidget::OnStreamingCheckboxChanged)
-					RENDER_RESOURCE_VIEWER_ADD_CHECKBOX(bShowRT, LOCTEXT("RTText", "RT"), LOCTEXT("RTTooltip", "Render Target"), SRenderResourceViewerWidget::OnRTCheckboxChanged)
-					RENDER_RESOURCE_VIEWER_ADD_CHECKBOX(bShowDS, LOCTEXT("DSText", "DS"), LOCTEXT("DSTooltip", "Depth Stencil"), SRenderResourceViewerWidget::OnDSCheckboxChanged)
-					RENDER_RESOURCE_VIEWER_ADD_CHECKBOX(bShowUAV, LOCTEXT("UAVText", "UAV"), LOCTEXT("UAVTooltip", "Unordered Access View"), SRenderResourceViewerWidget::OnUAVCheckboxChanged)
-					RENDER_RESOURCE_VIEWER_ADD_CHECKBOX(bShowRTAS, LOCTEXT("RTASText", "RTAS"), LOCTEXT("RTASTooltip", "Ray Tracing Acceleration Structure"), SRenderResourceViewerWidget::OnRTASCheckboxChanged)
+					RENDER_RESOURCE_VIEWER_ADD_CHECKBOX(bShowResident, LOCTEXT("ResidentText", "Resident"), LOCTEXT("ResidentTooltip", "Resource is accessible by GPU, and not evicted (unused)"), SRenderResourceViewerWidget::OnResidentCheckboxChanged)
+					RENDER_RESOURCE_VIEWER_ADD_CHECKBOX(bShowTransient, LOCTEXT("TransientText", "Transient"), LOCTEXT("TransientTooltip", "Resource is only allocated during the duration of the render passes where it's active and will share underlying memory with other resources in the frame"), SRenderResourceViewerWidget::OnTransientCheckboxChanged)
+					RENDER_RESOURCE_VIEWER_ADD_CHECKBOX(bShowStreaming, LOCTEXT("StreamingText", "Streaming"), LOCTEXT("StreamingTooltip", "Resource is a streamable texture"), SRenderResourceViewerWidget::OnStreamingCheckboxChanged)
+					RENDER_RESOURCE_VIEWER_ADD_CHECKBOX(bShowRT, LOCTEXT("RTText", "RT"), LOCTEXT("RTTooltip", "Resource can be written to as a Render Target buffer by GPU"), SRenderResourceViewerWidget::OnRTCheckboxChanged)
+					RENDER_RESOURCE_VIEWER_ADD_CHECKBOX(bShowDS, LOCTEXT("DSText", "DS"), LOCTEXT("DSTooltip", "Resource can be written to as a Depth Stencil buffer by GPU"), SRenderResourceViewerWidget::OnDSCheckboxChanged)
+					RENDER_RESOURCE_VIEWER_ADD_CHECKBOX(bShowUAV, LOCTEXT("UAVText", "UAV"), LOCTEXT("UAVTooltip", "Resource supports Unordered Access View which allows temporally unordered read/write access from multiple GPU threads without generating memory conflicts"), SRenderResourceViewerWidget::OnUAVCheckboxChanged)
+					RENDER_RESOURCE_VIEWER_ADD_CHECKBOX(bShowRTAS, LOCTEXT("RTASText", "RTAS"), LOCTEXT("RTASTooltip", "Resource is a Ray Tracing Acceleration Structure"), SRenderResourceViewerWidget::OnRTASCheckboxChanged)
 					RENDER_RESOURCE_VIEWER_ADD_CHECKBOX(bShowNone, LOCTEXT("NoneText", "None"), LOCTEXT("NoneTooltip", "Resource with no flags set"), SRenderResourceViewerWidget::OnNoneCheckboxChanged)
 
 					// Refresh button to update the resource list
