Render resource viewer - add tooltips on flags checkbox

diff --git a/Engine/Source/Editor/RenderResourceViewer/Private/SRenderResourceViewer.cpp b/Engine/Source/Editor/RenderResourceViewer/Private/SRenderResourceViewer.cpp
--- a/Engine/Source/Editor/RenderResourceViewer/Private/SRenderResourceViewer.cpp
+++ b/Engine/Source/Editor/RenderResourceViewer/Private/SRenderResourceViewer.cpp
@@ -97,6 +97,7 @@ namespace RenderResourceViewerInternal
 		.Padding(FMargin(0.0f, 0.0f, 10.0f, 0.0f)) \
 		[ \
 			SNew(SCheckBox) \
+			.ToolTipText(Tooltip) \
 			.IsChecked(bIsChecked) \
 			.OnCheckStateChanged(this, &OnCheckStateChangedFunc) \
 		]
