RenderResourceViewer - Add context menu on selected resource, and context command to find asset in content browser (shortcut CTRL+B)

diff --git a/Engine/Source/Editor/RenderResourceViewer/Private/SRenderResourceViewer.cpp b/Engine/Source/Editor/RenderResourceViewer/Private/SRenderResourceViewer.cpp
--- a/Engine/Source/Editor/RenderResourceViewer/Private/SRenderResourceViewer.cpp
+++ b/Engine/Source/Editor/RenderResourceViewer/Private/SRenderResourceViewer.cpp
@@ -5,7 +5,11 @@
 #include "Widgets/Input/SCheckBox.h"
 #include "Widgets/Input/SButton.h"
 #include "Widgets/Docking/SDockTab.h"
+#include "Framework/Commands/Commands.h"
+#include "Framework/MultiBox/MultiBoxBuilder.h"
 #include "RHI.h"
+#include "ContentBrowserModule.h"
+#include "IContentBrowserSingleton.h"
 
 #define LOCTEXT_NAMESPACE "SRenderResourceView"
 
@@ -81,6 +85,25 @@ namespace RenderResourceViewerInternal
 			return SNullWidget::NullWidget;
 		}
 	};
+
+	class FContextMenuCommands : public TCommands<FContextMenuCommands>
+	{
+	public:
+		FContextMenuCommands()
+			: TCommands<FContextMenuCommands>(TEXT("RenderResourceViewer"), NSLOCTEXT("Contexts", "RenderResourceViewer", "Render Resource Viewer"), NAME_None, FAppStyle::GetAppStyleSetName())
+		{}
+
+		virtual void RegisterCommands() override
+		{
+			UI_COMMAND(Command_FindInContentBrowser,
+				"Browse to Asset", 
+				"Browses to the associated asset and selects it in the most recently used Content Browser (summoning one if necessary)", 
+				EUserInterfaceActionType::Button, 
+				FInputChord(EModifierKey::Control, EKeys::B));
+		}
+
+		TSharedPtr<FUICommandInfo> Command_FindInContentBrowser;
+	};
 }
 
 #define RENDER_RESOURCE_VIEWER_ADD_CHECKBOX(bIsChecked, DisplayText, Tooltip, OnCheckStateChangedFunc) \
@@ -218,6 +241,7 @@ void SRenderResourceViewerWidget::Construct(const FArguments& InArgs, const TSha
 					.ListItemsSource(&ResourceInfos)
 					.SelectionMode(ESelectionMode::SingleToggle)
 					.OnGenerateRow(this, &SRenderResourceViewerWidget::HandleResourceGenerateRow)
+					.OnContextMenuOpening(FOnContextMenuOpening::CreateSP(this, &SRenderResourceViewerWidget::OpenContextMenu))
 					.HeaderRow
 					(
 						SNew(SHeaderRow)
@@ -233,6 +257,8 @@ void SRenderResourceViewerWidget::Construct(const FArguments& InArgs, const TSha
 		]
 	];
 
+	InitCommandList();
+
 	RefreshNodes(true);
 }
 
@@ -337,4 +363,65 @@ FText SRenderResourceViewerWidget::GetResourceSizeText() const
 	return FText::FromString(RenderResourceViewerInternal::GetFormatedSize(TotalResourceSize));
 }
 
+FReply SRenderResourceViewerWidget::OnKeyDown(const FGeometry& MyGeometry, const FKeyEvent& InKeyEvent)
+{
+	return CommandList->ProcessCommandBindings(InKeyEvent) == true ? FReply::Handled() : FReply::Unhandled();
+}
+
+void SRenderResourceViewerWidget::InitCommandList()
+{
+	RenderResourceViewerInternal::FContextMenuCommands::Register();
+	CommandList = MakeShared<FUICommandList>();
+
+	CommandList->MapAction(RenderResourceViewerInternal::FContextMenuCommands::Get().Command_FindInContentBrowser,
+		FExecuteAction::CreateSP(this, &SRenderResourceViewerWidget::ContextMenu_FindInContentBrowser),
+		FCanExecuteAction::CreateSP(this, &SRenderResourceViewerWidget::ContextMenu_FindInContentBrowser_CanExecute));
+}
+
+TSharedPtr<SWidget> SRenderResourceViewerWidget::OpenContextMenu()
+{
+	FMenuBuilder MenuBuilder(true, CommandList.ToSharedRef());
+	MenuBuilder.AddMenuEntry
+	(
+		RenderResourceViewerInternal::FContextMenuCommands::Get().Command_FindInContentBrowser,
+		NAME_None,
+		TAttribute<FText>(),
+		TAttribute<FText>(),
+		FSlateIcon(FAppStyle::GetAppStyleSetName(), "SystemWideCommands.FindInContentBrowser.Small")
+	);
+
+	return MenuBuilder.MakeWidget();
+}
+
+void SRenderResourceViewerWidget::ContextMenu_FindInContentBrowser()
+{
+	UObject* SelectedAsset = nullptr;
+	const TArray<TSharedPtr<FRHIResourceStats>>& SelectedNodes = ResourceListView->GetSelectedItems();
+	if (SelectedNodes.Num() > 0)
+	{
+		// Find the UObject asset from the owner name path
+		FString ObjectPathString = SelectedNodes[0]->OwnerName.ToString();
+		int32 LODIdx = ObjectPathString.Find(TEXT(" [LOD"));
+		if (LODIdx > -1)
+		{
+			ObjectPathString = ObjectPathString.Left(LODIdx);
+		}
+
+		SelectedAsset = FSoftObjectPath(ObjectPathString).ResolveObject();
+	}
+
+	if (SelectedAsset)
+	{
+		// Highlight the asset in content browser
+		const TArray<UObject*>& Assets = { SelectedAsset };
+		const FContentBrowserModule& ContentBrowserModule = FModuleManager::Get().LoadModuleChecked<FContentBrowserModule>("ContentBrowser");
+		ContentBrowserModule.Get().SyncBrowserToAssets(Assets);
+	}
+}
+
+bool SRenderResourceViewerWidget::ContextMenu_FindInContentBrowser_CanExecute() const
+{
+	return ResourceListView->GetSelectedItems().Num() > 0;
+}
+
 #undef LOCTEXT_NAMESPACE
diff --git a/Engine/Source/Editor/RenderResourceViewer/Private/SRenderResourceViewer.h b/Engine/Source/Editor/RenderResourceViewer/Private/SRenderResourceViewer.h
--- a/Engine/Source/Editor/RenderResourceViewer/Private/SRenderResourceViewer.h
+++ b/Engine/Source/Editor/RenderResourceViewer/Private/SRenderResourceViewer.h
@@ -12,6 +12,7 @@ struct FRHIResourceStats;
 class ITableRow;
 class STableViewBase;
 class SDockTab;
+class FUICommandList;
 
 class SRenderResourceViewerWidget : public SCompoundWidget
 {
@@ -22,6 +23,8 @@ public:
 	void Construct(const FArguments& InArgs, const TSharedRef<SDockTab>& ConstructUnderMajorTab, const TSharedPtr<SWindow>& ConstructUnderWindow);
 	~SRenderResourceViewerWidget() {}
 
+	virtual FReply OnKeyDown(const FGeometry& MyGeometry, const FKeyEvent& InKeyEvent) override;
+
 private:
 	void RefreshNodes(bool bUpdateRHIResources = false);
 	TSharedRef<ITableRow> HandleResourceGenerateRow(TSharedPtr<FRHIResourceStats> Item, const TSharedRef<STableViewBase>& OwnerTable);
@@ -40,6 +43,10 @@ private:
 	void OnRTASCheckboxChanged(ECheckBoxState NewState)					{ bShowRTAS = (NewState == ECheckBoxState::Checked); RefreshNodes(); }
 	void OnNoneCheckboxChanged(ECheckBoxState NewState)					{ bShowNone = (NewState == ECheckBoxState::Checked); RefreshNodes(); }
 	FReply OnRefreshButtonClicked()										{ RefreshNodes(true); return FReply::Handled(); }
+	void InitCommandList();
+	TSharedPtr<SWidget> OpenContextMenu();
+	void ContextMenu_FindInContentBrowser();
+	bool ContextMenu_FindInContentBrowser_CanExecute() const;
 
 	TArray<TSharedPtr<FRHIResourceStats>> RHIResources;
 	TSharedPtr<SListView<TSharedPtr<FRHIResourceStats>>> ResourceListView;
@@ -50,6 +57,7 @@ private:
 	TSharedPtr< class SEditableTextBox > FilterTextBox;
 	uint64 TotalResourceCount = 0;
 	uint64 TotalResourceSize = 0;
+	TSharedPtr<FUICommandList> CommandList;
 
 	bool bShowResident = true;			// Show resource with Resident flag set
 	bool bShowTransient = true;			// Show resource with Transient flag set
diff --git a/Engine/Source/Editor/RenderResourceViewer/RenderResourceViewer.Build.cs b/Engine/Source/Editor/RenderResourceViewer/RenderResourceViewer.Build.cs
--- a/Engine/Source/Editor/RenderResourceViewer/RenderResourceViewer.Build.cs
+++ b/Engine/Source/Editor/RenderResourceViewer/RenderResourceViewer.Build.cs
@@ -10,6 +10,9 @@ public class RenderResourceViewer : ModuleRules
 			new string[]
 			{
 				"WorkspaceMenuStructure",
+				"ContentBrowser",
+				"UnrealEd",
+				"AssetDefinition"
 			}
 		);
 
@@ -17,6 +20,7 @@ public class RenderResourceViewer : ModuleRules
 			new string[]
 			{
 				"Core",
+				"CoreUObject",
 				"InputCore",
 				"SlateCore",
 				"Slate",
