[Render Resource Viewer] Change the behaviour of flag checkbox based on user feedback. When unticked the resources with the flag set are now hidden from display. Also remove None flag checkbox and hide Transient resources by default.

diff --git a/Engine/Source/Editor/RenderResourceViewer/Private/SRenderResourceViewer.cpp b/Engine/Source/Editor/RenderResourceViewer/Private/SRenderResourceViewer.cpp
--- a/Engine/Source/Editor/RenderResourceViewer/Private/SRenderResourceViewer.cpp
+++ b/Engine/Source/Editor/RenderResourceViewer/Private/SRenderResourceViewer.cpp
@@ -169,7 +169,6 @@ void SRenderResourceViewerWidget::Construct(const FArguments& InArgs, const TSha
 					RENDER_RESOURCE_VIEWER_ADD_CHECKBOX(bShowDS, LOCTEXT("DSText", "DS"), LOCTEXT("DSTooltip", "Resource can be written to as a Depth Stencil buffer by GPU"), SRenderResourceViewerWidget::OnDSCheckboxChanged)
 					RENDER_RESOURCE_VIEWER_ADD_CHECKBOX(bShowUAV, LOCTEXT("UAVText", "UAV"), LOCTEXT("UAVTooltip", "Resource supports Unordered Access View which allows temporally unordered read/write access from multiple GPU threads without generating memory conflicts"), SRenderResourceViewerWidget::OnUAVCheckboxChanged)
 					RENDER_RESOURCE_VIEWER_ADD_CHECKBOX(bShowRTAS, LOCTEXT("RTASText", "RTAS"), LOCTEXT("RTASTooltip", "Resource is a Ray Tracing Acceleration Structure"), SRenderResourceViewerWidget::OnRTASCheckboxChanged)
-					RENDER_RESOURCE_VIEWER_ADD_CHECKBOX(bShowNone, LOCTEXT("NoneText", "None"), LOCTEXT("NoneTooltip", "Resource with no flags set"), SRenderResourceViewerWidget::OnNoneCheckboxChanged)
 
 					// Refresh button to update the resource list
 					+ SHorizontalBox::Slot()
@@ -280,13 +279,20 @@ void SRenderResourceViewerWidget::RefreshNodes(bool bUpdateRHIResources)
 	TotalResourceSize = 0;
 	for (const TSharedPtr<FRHIResourceStats>& Info : RHIResources)
 	{
+		ensure(Info->bHasFlags);
+
 		bool bContainsFilterText = FilterText.IsEmpty() || Info->Name.ToString().Contains(FilterText.ToString()) || Info->OwnerName.ToString().Contains(FilterText.ToString());
-		// Note bMarkedForDelete resources are excluded from display
-		bool bContainsFilterFlags = (bShowResident && Info->bResident) || (bShowTransient && Info->bTransient) || (bShowStreaming && Info->bStreaming)
-									|| (bShowRT && Info->bRenderTarget) || (bShowDS && Info->bDepthStencil) || (bShowUAV && Info->bUnorderedAccessView) || (bShowRTAS && Info->bRayTracingAccelerationStructure)
-									|| (bShowNone && !Info->bHasFlags);
+		bool bContainsFilterFlags = true;
+		// If the resource has a flag set and its matching check box is un-ticked, exclude from display.
 		// bMarkedForDelete resources are excluded from display
-		if (Info->bMarkedForDelete)
+		if ((!bShowResident && Info->bResident) ||
+			(!bShowTransient && Info->bTransient) ||
+			(!bShowStreaming && Info->bStreaming) ||
+			(!bShowRT && Info->bRenderTarget) ||
+			(!bShowDS && Info->bDepthStencil) ||
+			(!bShowUAV && Info->bUnorderedAccessView) ||
+			(!bShowRTAS && Info->bRayTracingAccelerationStructure) ||
+			Info->bMarkedForDelete)
 		{
 			bContainsFilterFlags = false;
 		}
diff --git a/Engine/Source/Editor/RenderResourceViewer/Private/SRenderResourceViewer.h b/Engine/Source/Editor/RenderResourceViewer/Private/SRenderResourceViewer.h
--- a/Engine/Source/Editor/RenderResourceViewer/Private/SRenderResourceViewer.h
+++ b/Engine/Source/Editor/RenderResourceViewer/Private/SRenderResourceViewer.h
@@ -41,7 +41,6 @@ private:
 	void OnDSCheckboxChanged(ECheckBoxState NewState)					{ bShowDS = (NewState == ECheckBoxState::Checked); RefreshNodes(); }
 	void OnUAVCheckboxChanged(ECheckBoxState NewState)					{ bShowUAV = (NewState == ECheckBoxState::Checked); RefreshNodes(); }
 	void OnRTASCheckboxChanged(ECheckBoxState NewState)					{ bShowRTAS = (NewState == ECheckBoxState::Checked); RefreshNodes(); }
-	void OnNoneCheckboxChanged(ECheckBoxState NewState)					{ bShowNone = (NewState == ECheckBoxState::Checked); RefreshNodes(); }
 	FReply OnRefreshButtonClicked()										{ RefreshNodes(true); return FReply::Handled(); }
 	void InitCommandList();
 	TSharedPtr<SWidget> OpenContextMenu();
@@ -60,11 +59,10 @@ private:
 	TSharedPtr<FUICommandList> CommandList;
 
 	bool bShowResident = true;			// Show resource with Resident flag set
-	bool bShowTransient = true;			// Show resource with Transient flag set
+	bool bShowTransient = false;		// Show resource with Transient flag set
 	bool bShowStreaming = true;			// Show resource with Streaming flag set
 	bool bShowRT = true;				// Show resource with RenderTarget flag set
 	bool bShowDS = true;				// Show resource with DepthStencil flag set
 	bool bShowUAV = true;				// Show resource with UAV flag set
 	bool bShowRTAS = true;				// Show resource with RayTracingAccelationStructure flag set
-	bool bShowNone = true;				// Show resource with no flag set
 };
