[Render Resource Viewer] Add tooltip to list items

diff --git a/Engine/Source/Editor/RenderResourceViewer/Private/SRenderResourceViewer.cpp b/Engine/Source/Editor/RenderResourceViewer/Private/SRenderResourceViewer.cpp
--- a/Engine/Source/Editor/RenderResourceViewer/Private/SRenderResourceViewer.cpp
+++ b/Engine/Source/Editor/RenderResourceViewer/Private/SRenderResourceViewer.cpp
@@ -237,6 +237,7 @@ void SRenderResourceViewerWidget::Construct(const FArguments& InArgs, const TSha
 					SAssignNew(ResourceListView, SListView<TSharedPtr<FRHIResourceStats>>)
 					.ScrollbarVisibility(EVisibility::Visible)
 					.ItemHeight(24.0f)
+					.ToolTipText(LOCTEXT("ResourceListViewToolTip", "CTRL+B or use right click menu to find the asset in Content Browser"))
 					.ListItemsSource(&ResourceInfos)
 					.SelectionMode(ESelectionMode::SingleToggle)
 					.OnGenerateRow(this, &SRenderResourceViewerWidget::HandleResourceGenerateRow)
