Remove ensure on bFlagSet - it's valid for buffers to not have any flags set

diff --git a/Engine/Source/Editor/RenderResourceViewer/Private/SRenderResourceViewer.cpp b/Engine/Source/Editor/RenderResourceViewer/Private/SRenderResourceViewer.cpp
--- a/Engine/Source/Editor/RenderResourceViewer/Private/SRenderResourceViewer.cpp
+++ b/Engine/Source/Editor/RenderResourceViewer/Private/SRenderResourceViewer.cpp
@@ -280,8 +280,6 @@ void SRenderResourceViewerWidget::RefreshNodes(bool bUpdateRHIResources)
 	TotalResourceSize = 0;
 	for (const TSharedPtr<FRHIResourceStats>& Info : RHIResources)
 	{
-		ensure(Info->bHasFlags);
-
 		bool bContainsFilterText = FilterText.IsEmpty() || Info->Name.ToString().Contains(FilterText.ToString()) || Info->OwnerName.ToString().Contains(FilterText.ToString());
 		bool bContainsFilterFlags = true;
 		// If the resource has a flag set and its matching check box is un-ticked, exclude from display.
