Change checkbox to combobox with Any,Yes & No options to make filtering of what to show in the resource list view easier

diff --git a/Engine/Source/Editor/RenderResourceViewer/Private/SRenderResourceViewer.cpp b/Engine/Source/Editor/RenderResourceViewer/Private/SRenderResourceViewer.cpp
--- a/Engine/Source/Editor/RenderResourceViewer/Private/SRenderResourceViewer.cpp
+++ b/Engine/Source/Editor/RenderResourceViewer/Private/SRenderResourceViewer.cpp
@@ -2,8 +2,8 @@
 
 #include "SRenderResourceViewer.h"
 #include "Widgets/Input/SSearchBox.h"
-#include "Widgets/Input/SCheckBox.h"
 #include "Widgets/Input/SButton.h"
+#include "Widgets/Input/STextComboBox.h"
 #include "Widgets/Docking/SDockTab.h"
 #include "Framework/Commands/Commands.h"
 #include "Framework/MultiBox/MultiBoxBuilder.h"
@@ -106,7 +106,7 @@ namespace RenderResourceViewerInternal
 	};
 }
 
-#define RENDER_RESOURCE_VIEWER_ADD_CHECKBOX(bIsChecked, DisplayText, Tooltip, OnCheckStateChangedFunc) \
+#define RENDER_RESOURCE_VIEWER_ADD_COMBOBOX(DisplayText, Tooltip, Options, SelectedItem, OnSelectionChangedFunc) \
 	+ SHorizontalBox::Slot() \
 		.AutoWidth() \
 		.Padding(FMargin(0.0f, 2.0f, 2.0f, 0.0f)) \
@@ -119,10 +119,11 @@ namespace RenderResourceViewerInternal
 		.AutoWidth() \
 		.Padding(FMargin(0.0f, 0.0f, 10.0f, 0.0f)) \
 		[ \
-			SNew(SCheckBox) \
+			SNew(STextComboBox) \
 			.ToolTipText(Tooltip) \
-			.IsChecked(bIsChecked) \
-			.OnCheckStateChanged(this, &OnCheckStateChangedFunc) \
+			.OptionsSource(&Options) \
+			.InitiallySelectedItem(SelectedItem) \
+			.OnSelectionChanged(this, &OnSelectionChangedFunc) \
 		]
 
 #define RENDER_RESOURCE_VIEWER_ADD_COLUMN(Name, Width, Label) \
@@ -137,6 +138,11 @@ void SRenderResourceViewerWidget::Construct(const FArguments& InArgs, const TSha
 	SortByColumn = RenderResourceViewerInternal::ColumnSize;
 	SortMode = EColumnSortMode::Descending;
 
+	ComboBoxNames.Empty();
+	ComboBoxNames.Add(MakeShareable(new FString(TEXT("Any"))));
+	ComboBoxNames.Add(MakeShareable(new FString(TEXT("Yes"))));
+	ComboBoxNames.Add(MakeShareable(new FString(TEXT("No"))));
+
 	ChildSlot
 	[
 		SNew(SVerticalBox)
@@ -161,14 +167,21 @@ void SRenderResourceViewerWidget::Construct(const FArguments& InArgs, const TSha
 						.Text(LOCTEXT("FilterText", "FilterByFlags:"))
 					]
 
-					// A row of checkboxes to filter resource list by flags
-					RENDER_RESOURCE_VIEWER_ADD_CHECKBOX(bShowResident, LOCTEXT("ResidentText", "Resident"), LOCTEXT("ResidentTooltip", "Resource is accessible by GPU, and not evicted (unused)"), SRenderResourceViewerWidget::OnResidentCheckboxChanged)
-					RENDER_RESOURCE_VIEWER_ADD_CHECKBOX(bShowTransient, LOCTEXT("TransientText", "Transient"), LOCTEXT("TransientTooltip", "Resource is only allocated during the duration of the render passes where it's active and will share underlying memory with other resources in the frame"), SRenderResourceViewerWidget::OnTransientCheckboxChanged)
-					RENDER_RESOURCE_VIEWER_ADD_CHECKBOX(bShowStreaming, LOCTEXT("StreamingText", "Streaming"), LOCTEXT("StreamingTooltip", "Resource is a streamable texture"), SRenderResourceViewerWidget::OnStreamingCheckboxChanged)
-					RENDER_RESOURCE_VIEWER_ADD_CHECKBOX(bShowRT, LOCTEXT("RTText", "RT"), LOCTEXT("RTTooltip", "Resource can be written to as a Render Target buffer by GPU"), SRenderResourceViewerWidget::OnRTCheckboxChanged)
-					RENDER_RESOURCE_VIEWER_ADD_CHECKBOX(bShowDS, LOCTEXT("DSText", "DS"), LOCTEXT("DSTooltip", "Resource can be written to as a Depth Stencil buffer by GPU"), SRenderResourceViewerWidget::OnDSCheckboxChanged)
-					RENDER_RESOURCE_VIEWER_ADD_CHECKBOX(bShowUAV, LOCTEXT("UAVText", "UAV"), LOCTEXT("UAVTooltip", "Resource supports Unordered Access View which allows temporally unordered read/write access from multiple GPU threads without generating memory conflicts"), SRenderResourceViewerWidget::OnUAVCheckboxChanged)
-					RENDER_RESOURCE_VIEWER_ADD_CHECKBOX(bShowRTAS, LOCTEXT("RTASText", "RTAS"), LOCTEXT("RTASTooltip", "Resource is a Ray Tracing Acceleration Structure"), SRenderResourceViewerWidget::OnRTASCheckboxChanged)
+					// A row of comboboxes to filter resource list by flags
+					RENDER_RESOURCE_VIEWER_ADD_COMBOBOX(LOCTEXT("ResidentText", "Resident"), LOCTEXT("ResidentTooltip", "Resource is accessible by GPU, and not evicted (unused). 'Any' will show all resources, 'Yes' only resources which are resident and 'No' resources which are evicted."), 
+						ComboBoxNames, ComboBoxTypeToName(ShowResident), SRenderResourceViewerWidget::OnResidentComboboxChanged)
+					RENDER_RESOURCE_VIEWER_ADD_COMBOBOX(LOCTEXT("TransientText", "Transient"), LOCTEXT("TransientTooltip", "Resource is only allocated during the duration of the render passes where it's active and will share underlying memory with other resources in the frame. 'Any' will show all resources, 'Yes' only resources which are transient and 'No' resources which are not transient."), 
+						ComboBoxNames, ComboBoxTypeToName(ShowTransient), SRenderResourceViewerWidget::OnTransientComboboxChanged)
+					RENDER_RESOURCE_VIEWER_ADD_COMBOBOX(LOCTEXT("StreamingText", "Streaming"), LOCTEXT("StreamingTooltip", "Resource is a streamable texture. 'Any' will show all resources, 'Yes' only resources which are streaming and 'No' resources which are not streaming."), 
+						ComboBoxNames, ComboBoxTypeToName(ShowStreaming), SRenderResourceViewerWidget::OnStreamingComboboxChanged)
+					RENDER_RESOURCE_VIEWER_ADD_COMBOBOX(LOCTEXT("RTText", "RT"), LOCTEXT("RTTooltip", "Resource can be written to as a Render Target buffer by GPU. 'Any' will show all resources, 'Yes' only resources which can be used as a RenderTarget and 'No' resources which can't be used as a RenderTarget."), 
+						ComboBoxNames, ComboBoxTypeToName(ShowRT), SRenderResourceViewerWidget::OnRTComboboxChanged)
+					RENDER_RESOURCE_VIEWER_ADD_COMBOBOX(LOCTEXT("DSText", "DS"), LOCTEXT("DSTooltip", "Resource can be written to as a Depth Stencil buffer by GPU. 'Any' will show all resources, 'Yes' only resources which can be used as a DepthStencil Target and 'No' resources which can't be used as a DepthStencil Target."), 
+						ComboBoxNames, ComboBoxTypeToName(ShowDS), SRenderResourceViewerWidget::OnDSComboboxChanged)
+					RENDER_RESOURCE_VIEWER_ADD_COMBOBOX(LOCTEXT("UAVText", "UAV"), LOCTEXT("UAVTooltip", "Resource supports Unordered Access View which allows temporally unordered read/write access from multiple GPU threads without generating memory conflicts. . 'Any' will show all resources, 'Yes' only resources which can be used as a UAV and 'No' resources which can't be used as a UAV."), 
+						ComboBoxNames, ComboBoxTypeToName(ShowUAV), SRenderResourceViewerWidget::OnUAVComboboxChanged)
+					RENDER_RESOURCE_VIEWER_ADD_COMBOBOX(LOCTEXT("RTASText", "RTAS"), LOCTEXT("RTASTooltip", "Resource is a Ray Tracing Acceleration Structure. 'Any' will show all resources, 'Yes' only resources which are RTAS and 'No' resources which are not RTAS."), 
+						ComboBoxNames, ComboBoxTypeToName(ShowRTAS), SRenderResourceViewerWidget::OnRTASComboboxChanged)
 
 					// Refresh button to update the resource list
 					+ SHorizontalBox::Slot()
@@ -262,6 +275,37 @@ void SRenderResourceViewerWidget::Construct(const FArguments& InArgs, const TSha
 	RefreshNodes(true);
 }
 
+SRenderResourceViewerWidget::EComboBoxType SRenderResourceViewerWidget::ComboBoxNameToType(TSharedPtr<FString> Value) const
+{
+	int32 Index = ComboBoxNames.Find(Value);
+	if (Index != INDEX_NONE)
+	{
+		return (EComboBoxType)Index;		
+	}
+	return EComboBoxType::Any;
+}
+
+bool SRenderResourceViewerWidget::ShouldShow(EComboBoxType FilterType, bool bValue) const
+{
+	switch (FilterType)
+	{
+	case EComboBoxType::Any:
+	{
+		return true;
+	}
+	case EComboBoxType::No:
+	{
+		return !bValue;
+	}
+	case EComboBoxType::Yes:
+	{
+		return bValue;
+	}
+	}
+
+	return true;
+}
+
 void SRenderResourceViewerWidget::RefreshNodes(bool bUpdateRHIResources)
 {
 	const int32 NumberOfResourcesToShow = -1;
@@ -284,13 +328,13 @@ void SRenderResourceViewerWidget::RefreshNodes(bool bUpdateRHIResources)
 		bool bContainsFilterFlags = true;
 		// If the resource has a flag set and its matching check box is un-ticked, exclude from display.
 		// bMarkedForDelete resources are excluded from display
-		if ((!bShowResident && Info->bResident) ||
-			(!bShowTransient && Info->bTransient) ||
-			(!bShowStreaming && Info->bStreaming) ||
-			(!bShowRT && Info->bRenderTarget) ||
-			(!bShowDS && Info->bDepthStencil) ||
-			(!bShowUAV && Info->bUnorderedAccessView) ||
-			(!bShowRTAS && Info->bRayTracingAccelerationStructure) ||
+		if (!ShouldShow(ShowResident, Info->bResident) ||
+			!ShouldShow(ShowTransient, Info->bTransient) ||
+			!ShouldShow(ShowStreaming, Info->bStreaming) ||
+			!ShouldShow(ShowRT, Info->bRenderTarget) ||
+			!ShouldShow(ShowDS, Info->bDepthStencil) ||
+			!ShouldShow(ShowUAV, Info->bUnorderedAccessView) ||
+			!ShouldShow(ShowRTAS, Info->bRayTracingAccelerationStructure) ||
 			Info->bMarkedForDelete)
 		{
 			bContainsFilterFlags = false;
diff --git a/Engine/Source/Editor/RenderResourceViewer/Private/SRenderResourceViewer.h b/Engine/Source/Editor/RenderResourceViewer/Private/SRenderResourceViewer.h
--- a/Engine/Source/Editor/RenderResourceViewer/Private/SRenderResourceViewer.h
+++ b/Engine/Source/Editor/RenderResourceViewer/Private/SRenderResourceViewer.h
@@ -34,19 +34,29 @@ private:
 	FText GetResourceCountText() const									{ return FText::AsNumber(TotalResourceCount); }
 	FText GetResourceSizeText() const;
 	void FilterTextBox_OnTextChanged(const FText& InFilterText)			{ FilterText = InFilterText; RefreshNodes(); }
-	void OnResidentCheckboxChanged(ECheckBoxState NewState)				{ bShowResident = (NewState == ECheckBoxState::Checked); RefreshNodes(); }
-	void OnTransientCheckboxChanged(ECheckBoxState NewState)			{ bShowTransient = (NewState == ECheckBoxState::Checked); RefreshNodes(); }
-	void OnStreamingCheckboxChanged(ECheckBoxState NewState)			{ bShowStreaming = (NewState == ECheckBoxState::Checked); RefreshNodes(); }
-	void OnRTCheckboxChanged(ECheckBoxState NewState)					{ bShowRT = (NewState == ECheckBoxState::Checked); RefreshNodes(); }
-	void OnDSCheckboxChanged(ECheckBoxState NewState)					{ bShowDS = (NewState == ECheckBoxState::Checked); RefreshNodes(); }
-	void OnUAVCheckboxChanged(ECheckBoxState NewState)					{ bShowUAV = (NewState == ECheckBoxState::Checked); RefreshNodes(); }
-	void OnRTASCheckboxChanged(ECheckBoxState NewState)					{ bShowRTAS = (NewState == ECheckBoxState::Checked); RefreshNodes(); }
+	void OnResidentComboboxChanged(TSharedPtr<FString> NewValue, ESelectInfo::Type)		{ ShowResident = ComboBoxNameToType(NewValue); RefreshNodes(); }	
+	void OnTransientComboboxChanged(TSharedPtr<FString> NewValue, ESelectInfo::Type)	{ ShowTransient = ComboBoxNameToType(NewValue); RefreshNodes(); }
+	void OnStreamingComboboxChanged(TSharedPtr<FString> NewValue, ESelectInfo::Type)	{ ShowStreaming = ComboBoxNameToType(NewValue); RefreshNodes(); }
+	void OnRTComboboxChanged(TSharedPtr<FString> NewValue, ESelectInfo::Type)			{ ShowRT = ComboBoxNameToType(NewValue); RefreshNodes(); }
+	void OnDSComboboxChanged(TSharedPtr<FString> NewValue, ESelectInfo::Type)			{ ShowDS = ComboBoxNameToType(NewValue); RefreshNodes(); }
+	void OnUAVComboboxChanged(TSharedPtr<FString> NewValue, ESelectInfo::Type)			{ ShowUAV = ComboBoxNameToType(NewValue); RefreshNodes(); }
+	void OnRTASComboboxChanged(TSharedPtr<FString> NewValue, ESelectInfo::Type)			{ ShowRTAS = ComboBoxNameToType(NewValue); RefreshNodes(); }
 	FReply OnRefreshButtonClicked()										{ RefreshNodes(true); return FReply::Handled(); }
 	void InitCommandList();
 	TSharedPtr<SWidget> OpenContextMenu();
 	void ContextMenu_FindInContentBrowser();
 	bool ContextMenu_FindInContentBrowser_CanExecute() const;
 
+	enum class EComboBoxType : uint8
+	{
+		Any,
+		Yes,
+		No,
+	};
+	EComboBoxType ComboBoxNameToType(TSharedPtr<FString> Value) const;
+	TSharedPtr<FString> ComboBoxTypeToName(EComboBoxType Type) const	{ return ComboBoxNames[(int32)Type]; }
+	bool ShouldShow(EComboBoxType FilterType, bool bValue) const;
+
 	TArray<TSharedPtr<FRHIResourceStats>> RHIResources;
 	TSharedPtr<SListView<TSharedPtr<FRHIResourceStats>>> ResourceListView;
 	TArray<TSharedPtr<FRHIResourceStats>> ResourceInfos;
@@ -57,12 +67,13 @@ private:
 	uint64 TotalResourceCount = 0;
 	uint64 TotalResourceSize = 0;
 	TSharedPtr<FUICommandList> CommandList;
+	TArray<TSharedPtr<FString>> ComboBoxNames;
 
-	bool bShowResident = true;			// Show resource with Resident flag set
-	bool bShowTransient = false;		// Show resource with Transient flag set
-	bool bShowStreaming = true;			// Show resource with Streaming flag set
-	bool bShowRT = true;				// Show resource with RenderTarget flag set
-	bool bShowDS = true;				// Show resource with DepthStencil flag set
-	bool bShowUAV = true;				// Show resource with UAV flag set
-	bool bShowRTAS = true;				// Show resource with RayTracingAccelationStructure flag set
+	EComboBoxType ShowResident = EComboBoxType::Any;					// Show resource with Resident flag set
+	EComboBoxType ShowTransient = EComboBoxType::No;					// Show resource with Transient flag set
+	EComboBoxType ShowStreaming = EComboBoxType::Any;					// Show resource with Streaming flag set
+	EComboBoxType ShowRT = EComboBoxType::Any;							// Show resource with RenderTarget flag set
+	EComboBoxType ShowDS = EComboBoxType::Any;							// Show resource with DepthStencil flag set
+	EComboBoxType ShowUAV = EComboBoxType::Any;							// Show resource with UAV flag set
+	EComboBoxType ShowRTAS = EComboBoxType::Any;						// Show resource with RayTracingAccelationStructure flag set
 };
