Adds basic TreeMap view to the Render Resource viewer

- The filters to control what buckets of data we show in the TreeMap for render resources are data driven via an ini file.

[/Script/RenderResourceViewer.RenderResourceViewerSettings]
+Filters=(FilterString="Lumen.",DisplayName="Lumen")

- The snippet above will find any RenderResources that start with "Lumen." and group them into a single bucket in the SizeMap
- SRenderResourceViewerWidget::RefreshSizeMap does the heavy lifting of building the SizeMap
- The SizeMap has a single root node called "GPU Memory" which everything else is a child to
- Anything that isn't picked up by the ini file filters will be placed in "Other"
- Added URenderResourceViewerSettings which are project settings to allow for the filters to be constructed in the Editor and saved w/ the project to share w/ others.
- Only the leaf nodes in the tree will have sizes, the parent nodes automatically have their size figured out based on children.
- There is a lot of string building in the SizeMap, but this function only runs once when you explicity refresh or open the RenderResourceViewer

todo: Probably a bunch of feature requests to make it more useable!

[RN] major | Rendering - Architecture

diff --git a/Engine/Config/BaseEditor.ini b/Engine/Config/BaseEditor.ini
--- a/Engine/Config/BaseEditor.ini
+++ b/Engine/Config/BaseEditor.ini
@@ -698,3 +698,13 @@ bUseWidgetTemplateSelector=True
 +WidgetClassesToHide=/Script/UMG.SpinBox
 +WidgetClassesToHide=/Script/UMG.Throbber
 +WidgetClassesToHide=/Script/UMG.TreeView
+
+[/Script/RenderResourceViewer.RenderResourceViewerSettings]
++Filters=(FilterString="FParticle",DisplayName="Particle Buffers")
++Filters=(FilterString="Lumen.",DisplayName="Lumen")
++Filters=(FilterString="Nanite.",DisplayName="Nanite")
++Filters=(FilterString="Shadow.",DisplayName="Shadows")
++Filters=(FilterString="SkyLight.",DisplayName="Skylight")
++Filters=(FilterString="TSR.",DisplayName="Temporal Super Resolution")
++Filters=(FilterString="TexCoordRHIBuffer",DisplayName="Tex Coord Buffers")
++Filters=(FilterString="FColorVertexBuffer",DisplayName="Color Vertex Buffers")
\ No newline at end of file
diff --git a/Engine/Source/Editor/RenderResourceViewer/Private/RenderResourceViewerSettings.cpp b/Engine/Source/Editor/RenderResourceViewer/Private/RenderResourceViewerSettings.cpp
new file mode 100644
--- /dev/null
+++ b/Engine/Source/Editor/RenderResourceViewer/Private/RenderResourceViewerSettings.cpp
@@ -0,0 +1,8 @@
+// Copyright Epic Games, Inc. All Rights Reserved.
+
+#include "RenderResourceViewerSettings.h"
+
+URenderResourceViewerSettings::URenderResourceViewerSettings(const FObjectInitializer& ObjectInitializer)
+	: Super(ObjectInitializer)
+{ 
+}
diff --git a/Engine/Source/Editor/RenderResourceViewer/Private/SRenderResourceViewer.cpp b/Engine/Source/Editor/RenderResourceViewer/Private/SRenderResourceViewer.cpp
--- a/Engine/Source/Editor/RenderResourceViewer/Private/SRenderResourceViewer.cpp
+++ b/Engine/Source/Editor/RenderResourceViewer/Private/SRenderResourceViewer.cpp
@@ -10,6 +10,8 @@
 #include "RHI.h"
 #include "ContentBrowserModule.h"
 #include "IContentBrowserSingleton.h"
+#include "Developer/TreeMap/Public/STreeMap.h"
+#include "RenderResourceViewerSettings.h"
 
 #define LOCTEXT_NAMESPACE "SRenderResourceView"
 
@@ -133,6 +135,13 @@ namespace RenderResourceViewerInternal
 		.SortMode(this, &SRenderResourceViewerWidget::GetColumnSortMode, Name) \
 		.OnSort(this, &SRenderResourceViewerWidget::OnColumnSortModeChanged)
 
+SRenderResourceViewerWidget::SRenderResourceViewerWidget()
+: TreeMapWidget(nullptr),
+  RootTreeMapNode(new FTreeMapNodeData())
+{
+
+}
+
 void SRenderResourceViewerWidget::Construct(const FArguments& InArgs, const TSharedRef<SDockTab>& ConstructUnderMajorTab, const TSharedPtr<SWindow>& ConstructUnderWindow)
 {
 	SortByColumn = RenderResourceViewerInternal::ColumnSize;
@@ -153,7 +162,11 @@ void SRenderResourceViewerWidget::Construct(const FArguments& InArgs, const TSha
 			SNew(SBox)
 			[
 				SNew(SVerticalBox)
-
+				+ SVerticalBox::Slot()
+				.FillHeight(1.f)
+				[
+					SAssignNew(TreeMapWidget, STreeMap, RootTreeMapNode.ToSharedRef(), nullptr)
+				]
 				+ SVerticalBox::Slot()
 				.AutoHeight()
 				[
@@ -392,6 +405,8 @@ void SRenderResourceViewerWidget::RefreshNodes(bool bUpdateRHIResources)
 	}
 
 	ResourceListView->RebuildList();
+
+	RefreshSizeMap();
 }
 
 TSharedRef<ITableRow> SRenderResourceViewerWidget::HandleResourceGenerateRow(TSharedPtr<FRHIResourceStats> Item, const TSharedRef<STableViewBase>& OwnerTable)
@@ -472,4 +487,96 @@ bool SRenderResourceViewerWidget::ContextMenu_FindInContentBrowser_CanExecute()
 	return ResourceListView->GetSelectedItems().Num() > 0;
 }
 
+void SRenderResourceViewerWidget::RefreshSizeMap()
+{
+	RootTreeMapNode->Children.Empty();
+
+	const URenderResourceViewerSettings* Settings = GetDefault<URenderResourceViewerSettings>();
+
+	RootTreeMapNode->LogicalName = TEXT("Total");
+	RootTreeMapNode->Name = TEXT("GPU Memory");
+	RootTreeMapNode->Name2 = FString::Printf(TEXT("%s"), *FText::AsMemory(TotalResourceSize).ToString());
+
+	// Build parent nodes in the treemap for all our filters
+	TMap<FString, FTreeMapNodeDataPtr> NodeLookup;
+	for (const FRenderResourceViewerTreemapFilter& Filter : Settings->Filters)
+	{
+		FTreeMapNodeDataPtr TreeMapNode = MakeShareable(new FTreeMapNodeData());
+		NodeLookup.Add(Filter.FilterString, TreeMapNode);
+
+		TreeMapNode->LogicalName = Filter.FilterString;
+		TreeMapNode->Name = Filter.DisplayName;
+		TreeMapNode->Size = 0.0f;
+		RootTreeMapNode->Children.Add(TreeMapNode);
+		TreeMapNode->Parent = RootTreeMapNode.Get();
+	}
+
+	// Catch all filter for stuff that doesn't go into our filters.
+	FTreeMapNodeDataPtr OtherNode = MakeShareable(new FTreeMapNodeData());
+	OtherNode->LogicalName = TEXT("Other");
+	OtherNode->Name = TEXT("Other");
+	OtherNode->Size = 0.0f;
+	RootTreeMapNode->Children.Add(OtherNode);
+	OtherNode->Parent = RootTreeMapNode.Get();
+
+	for (TSharedPtr<FRHIResourceStats> Stat : ResourceInfos)
+	{
+		FTreeMapNodeDataPtr ChildNode = MakeShareable(new FTreeMapNodeData());
+		ChildNode->LogicalName = Stat->Name.ToString();
+		ChildNode->Name = Stat->Name.ToString();
+		ChildNode->Size = Stat->SizeInBytes;
+		ChildNode->Name2 = FString::Printf(TEXT("%s"), *FText::AsMemory(Stat->SizeInBytes).ToString());
+
+		bool bAddedToFilter = false;
+		for (const FRenderResourceViewerTreemapFilter& Filter : Settings->Filters)
+		{
+			if (Stat->Name.ToString().StartsWith(Filter.FilterString))
+			{
+				if (FTreeMapNodeDataPtr* ParentNode = NodeLookup.Find(Filter.FilterString))
+				{
+					ChildNode->Parent = (*ParentNode).Get();
+					(*ParentNode)->Children.Add(ChildNode);
+					bAddedToFilter = true;
+					break;
+				}
+			}
+		}
+
+		if (!bAddedToFilter)
+		{
+			// Try to add by owner name.
+			if (!Stat->OwnerName.IsNone())
+			{
+				FString OwnerName = Stat->OwnerName.ToString();
+				if (FTreeMapNodeDataPtr* ParentNode = NodeLookup.Find(OwnerName))
+				{
+					ChildNode->Parent = (*ParentNode).Get();
+					(*ParentNode)->Children.Add(ChildNode);
+				}
+				else
+				{
+					FTreeMapNodeDataPtr OwnerNode = MakeShareable(new FTreeMapNodeData());
+					OwnerNode->LogicalName = OwnerName;
+					OwnerNode->Name = OwnerName;
+					OwnerNode->Size = Stat->SizeInBytes;
+
+					RootTreeMapNode->Children.Add(OwnerNode);
+					OwnerNode->Parent = RootTreeMapNode.Get();
+
+					NodeLookup.Add(OwnerName, OwnerNode);
+				}
+			}
+			else
+			{
+				// Add it to the other bucket
+				ChildNode->Parent = OtherNode.Get();
+				OtherNode->Children.Add(ChildNode);
+			}
+		}
+	}
+
+	TreeMapWidget->RebuildTreeMap(false);
+	TreeMapWidget->ZoomOut();
+}
+
 #undef LOCTEXT_NAMESPACE
diff --git a/Engine/Source/Editor/RenderResourceViewer/Private/SRenderResourceViewer.h b/Engine/Source/Editor/RenderResourceViewer/Private/SRenderResourceViewer.h
--- a/Engine/Source/Editor/RenderResourceViewer/Private/SRenderResourceViewer.h
+++ b/Engine/Source/Editor/RenderResourceViewer/Private/SRenderResourceViewer.h
@@ -21,12 +21,15 @@ public:
 	SLATE_END_ARGS()
 
 	void Construct(const FArguments& InArgs, const TSharedRef<SDockTab>& ConstructUnderMajorTab, const TSharedPtr<SWindow>& ConstructUnderWindow);
+
+	SRenderResourceViewerWidget();
 	~SRenderResourceViewerWidget() {}
 
 	virtual FReply OnKeyDown(const FGeometry& MyGeometry, const FKeyEvent& InKeyEvent) override;
 
 private:
 	void RefreshNodes(bool bUpdateRHIResources = false);
+	void RefreshSizeMap();
 	TSharedRef<ITableRow> HandleResourceGenerateRow(TSharedPtr<FRHIResourceStats> Item, const TSharedRef<STableViewBase>& OwnerTable);
 	void OnColumnSortModeChanged(const EColumnSortPriority::Type SortPriority, const FName& ColumnId, const EColumnSortMode::Type InSortMode);
 
@@ -69,6 +72,11 @@ private:
 	TSharedPtr<FUICommandList> CommandList;
 	TArray<TSharedPtr<FString>> ComboBoxNames;
 
+	/** Our tree map widget */
+	TSharedPtr<class STreeMap> TreeMapWidget;
+	/** Our tree map source data */
+	TSharedPtr<class FTreeMapNodeData> RootTreeMapNode;
+
 	EComboBoxType ShowResident = EComboBoxType::Any;					// Show resource with Resident flag set
 	EComboBoxType ShowTransient = EComboBoxType::No;					// Show resource with Transient flag set
 	EComboBoxType ShowStreaming = EComboBoxType::Any;					// Show resource with Streaming flag set
diff --git a/Engine/Source/Editor/RenderResourceViewer/Public/RenderResourceViewerSettings.h b/Engine/Source/Editor/RenderResourceViewer/Public/RenderResourceViewerSettings.h
new file mode 100644
--- /dev/null
+++ b/Engine/Source/Editor/RenderResourceViewer/Public/RenderResourceViewerSettings.h
@@ -0,0 +1,28 @@
+// Copyright Epic Games, Inc. All Rights Reserved.
+
+#pragma once
+
+#include "CoreMinimal.h"
+#include "Engine/DeveloperSettings.h"
+#include "RenderResourceViewerSettings.generated.h"
+
+USTRUCT()
+struct FRenderResourceViewerTreemapFilter
+{
+	GENERATED_BODY()
+
+	UPROPERTY(config, EditAnywhere, Category = "Section")
+	FString FilterString;
+
+	UPROPERTY(config, EditAnywhere, Category = "Section")
+	FString DisplayName;
+};
+
+UCLASS(config = Editor, DefaultConfig, meta = (DisplayName = "Render Resource Viewer"))
+class RENDERRESOURCEVIEWER_API URenderResourceViewerSettings : public UDeveloperSettings
+{
+	GENERATED_UCLASS_BODY()
+
+	UPROPERTY(Config, EditAnywhere, Category = "Treemap")
+	TArray<FRenderResourceViewerTreemapFilter> Filters;
+};
diff --git a/Engine/Source/Editor/RenderResourceViewer/RenderResourceViewer.Build.cs b/Engine/Source/Editor/RenderResourceViewer/RenderResourceViewer.Build.cs
--- a/Engine/Source/Editor/RenderResourceViewer/RenderResourceViewer.Build.cs
+++ b/Engine/Source/Editor/RenderResourceViewer/RenderResourceViewer.Build.cs
@@ -23,6 +23,8 @@ public class RenderResourceViewer : ModuleRules
 				"SlateCore",
 				"Slate",
 				"RHI",
+				"DeveloperSettings",
+				"TreeMap",
 			}
 		);
 	}
